Copying //UE4/Partner-MagicLeap-Staging to Dev-VR (//UE4/Dev-VR)

diff --git a/Engine/Plugins/Lumin/MagicLeapScreens/Source/MagicLeapScreens/Private/MagicLeapScreensPlugin.h b/Engine/Plugins/Lumin/MagicLeapScreens/Source/MagicLeapScreens/Private/MagicLeapScreensPlugin.h
--- a/Engine/Plugins/Lumin/MagicLeapScreens/Source/MagicLeapScreens/Private/MagicLeapScreensPlugin.h
+++ b/Engine/Plugins/Lumin/MagicLeapScreens/Source/MagicLeapScreens/Private/MagicLeapScreensPlugin.h
@@ -8,53 +8,61 @@
 #include "Misc/Optional.h"
 #include "IMagicLeapPlugin.h"
 #include "IMagicLeapScreensPlugin.h"
-#include "MagicLeapUtils.h"
-#include "MagicLeapScreensMsg.h"
-#include "MagicLeapPluginUtil.h" // for ML_INCLUDES_START/END
+#include "Lumin/CAPIShims/LuminAPIImage.h"
+#include "Lumin/CAPIShims/LuminAPIScreens.h"
+#include "MagicLeapPluginUtil.h"
+#include "MagicLeapScreensTypes.h"
 
-#if WITH_MLSDK
-ML_INCLUDES_START
-#include <ml_image.h>
-#include <ml_screens.h>
-ML_INCLUDES_END
-#endif //WITH_MLSDK
+class FScreensRunnable;
+struct FScreensTask;
 
 class FMagicLeapScreensPlugin : public IMagicLeapScreensPlugin
 {
 public:
 	void StartupModule() override;
 	void ShutdownModule() override;
-	bool IsEngineLoopInitComplete() const override;
-	void OnEngineLoopInitComplete() override;
-	static bool IsSupportedFormat(EPixelFormat InPixelFormat);
+	bool IsSupportedFormat(EPixelFormat InPixelFormat);
 
 #if WITH_MLSDK
-	static UTexture2D* MLImageToUTexture2D(const MLImage& Source);
-	static void MLWatchHistoryEntryToUnreal(const MLScreensWatchHistoryEntry& InEntry, FScreensWatchHistoryEntry& OutEntry);
-	static bool UTexture2DToMLImage(const UTexture2D& Source, MLImage& Target);
+	UTexture2D* MLImageToUTexture2D(const MLImage& Source);
+	void MLWatchHistoryEntryToUnreal(const MLScreensWatchHistoryEntry& InEntry, FMagicLeapScreensWatchHistoryEntry& OutEntry);
+	bool MLScreenInfoToUnreal(const MLScreensScreenInfoEx& InInfo, FMagicLeapScreenTransform& OutTransform, const FTransform& TrackingToWorld, const float WorldToMetersScale);
+	bool UTexture2DToMLImage(const UTexture2D& Source, MLImage& Target);
 #endif // WITH_MLSDK
 
-	static bool RemoveWatchHistoryEntry(const FScreenID& ID);
-	static FScreensMessage GetWatchHistoryEntries();
-	static bool ClearWatchHistory();
+	bool RemoveWatchHistoryEntry(const FGuid& ID);
+	bool ClearWatchHistory();
 	bool Tick(float DeltaTime);
-	static FScreensMessage AddToWatchHistory(const FScreensWatchHistoryEntry& WatchHistoryEntry);
 
-	static void GetWatchHistoryEntriesAsync(const TOptional<FScreensHistoryRequestResultDelegate>& OptionalResultDelegate);
-	static void AddToWatchHistoryAsync(const FScreensWatchHistoryEntry& NewEntry, const TOptional<FScreensEntryRequestResultDelegate>& OptionalResultDelegate);
-	static void UpdateWatchHistoryEntryAsync(const FScreensWatchHistoryEntry& UpdateEntry, const TOptional<FScreensEntryRequestResultDelegate>& OptionalResultDelegate);
+	void GetWatchHistoryEntriesAsync(const TOptional<FMagicLeapScreensHistoryRequestResultDelegate>& OptionalResultDelegate);
+	void AddToWatchHistoryAsync(const FMagicLeapScreensWatchHistoryEntry& NewEntry, const TOptional<FMagicLeapScreensEntryRequestResultDelegate>& OptionalResultDelegate);
+	void UpdateWatchHistoryEntryAsync(const FMagicLeapScreensWatchHistoryEntry& UpdateEntry, const TOptional<FMagicLeapScreensEntryRequestResultDelegate>& OptionalResultDelegate);
+	void UpdateScreenTransformAsync(const FMagicLeapScreenTransform& ScreenTransform, const FMagicLeapScreenTransformRequestResultDelegate& ResultDelegate);
+
+	bool GetScreenTransform(FMagicLeapScreenTransform& ScreensTransforms);
+	bool GetScreensTransforms(TArray<FMagicLeapScreenTransform>& ScreensTransforms);
 
-	static FScreensMessage UpdateWatchHistoryEntry(const FScreensWatchHistoryEntry& WatchHistoryEntry);
-	static bool GetScreensTransforms(TArray<FScreenTransform>& ScreensTransforms);
-	static TArray<uint8> PixelDataMemPool;
-	static class FScreensWorker *Impl;
+	FScreensTask AddToWatchHistory(const FMagicLeapScreensWatchHistoryEntry& WatchHistoryEntry);
+	FScreensTask GetWatchHistoryEntries();
+	FScreensTask UpdateWatchHistoryEntry(const FMagicLeapScreensWatchHistoryEntry& WatchHistoryEntry);
+	FScreensTask UpdateScreensTransform(const FMagicLeapScreenTransform& ScreenTransform);
+	class FScreensRunnable* GetRunnable() const;
 
 private:
+	class FScreensRunnable* Runnable;
+	TArray<uint8> PixelDataMemPool;
 #if WITH_MLSDK
-	static bool ShouldUseDefaultThumbnail(const FScreensWatchHistoryEntry& Entry, MLImage& MLImage);
-	static MLImage DefaultThumbnail;
-	static FCriticalSection CriticalSection;
+	bool ShouldUseDefaultThumbnail(const FMagicLeapScreensWatchHistoryEntry& Entry, MLImage& MLImage);
+	MLImage DefaultThumbnail;
+	FCriticalSection CriticalSection;
 #endif // WITH_MLSDK
 	FMagicLeapAPISetup APISetup;
-	bool bEngineLoopInitComplete;
+
+	FTickerDelegate TickDelegate;
+	FDelegateHandle TickDelegateHandle;
 };
+
+inline FMagicLeapScreensPlugin& GetMagicLeapScreensPlugin()
+{
+	return FModuleManager::Get().GetModuleChecked<FMagicLeapScreensPlugin>("MagicLeapScreens");
+}
diff --git a/Engine/Plugins/Lumin/MagicLeapScreens/Source/MagicLeapScreens/Private/MagicLeapScreensRunnable.cpp b/Engine/Plugins/Lumin/MagicLeapScreens/Source/MagicLeapScreens/Private/MagicLeapScreensRunnable.cpp
new file mode 100644
--- /dev/null
+++ b/Engine/Plugins/Lumin/MagicLeapScreens/Source/MagicLeapScreens/Private/MagicLeapScreensRunnable.cpp
@@ -0,0 +1,69 @@
+// Copyright 1998-2019 Epic Games, Inc. All Rights Reserved.
+
+#include "MagicLeapScreensRunnable.h"
+#include "MagicLeapScreensPlugin.h"
+
+FScreensRunnable::FScreensRunnable()
+	: FMagicLeapRunnable({ }, TEXT("FScreensRunnable"))
+{
+}
+
+bool FScreensRunnable::PeekCompletedTasks(FScreensTask& OutTask) const
+{
+	return CompletedTasks.Peek(OutTask);
+}
+
+bool FScreensRunnable::CompletedTaskQueueIsEmpty() const
+{
+	return CompletedTasks.IsEmpty();
+}
+
+bool FScreensRunnable::ProcessCurrentTask()
+{
+	switch (CurrentTask.TaskType)
+	{
+	case FScreensTask::ETaskType::None: return true;
+	case FScreensTask::ETaskType::GetHistory: return GetWatchHistory();
+	case FScreensTask::ETaskType::AddToHistory: return AddToHistory();
+	case FScreensTask::ETaskType::UpdateHistoryEntry: return UpdateWatchHistoryEntry();
+	case FScreensTask::ETaskType::UpdateInfoEntry: return UpdateInfoEntry();
+	}
+	return false;
+}
+
+bool FScreensRunnable::AddToHistory()
+{
+	check(CurrentTask.WatchHistory.Num() != 0);
+	FScreensTask Task = GetMagicLeapScreensPlugin().AddToWatchHistory(CurrentTask.WatchHistory[0]);
+	Task.EntryRequestDelegate = CurrentTask.EntryRequestDelegate;
+	CurrentTask = Task;
+	check(CurrentTask.WatchHistory.Num() > 0);
+	return CurrentTask.bSuccess;
+}
+
+bool FScreensRunnable::UpdateWatchHistoryEntry()
+{
+	check(CurrentTask.WatchHistory.Num() != 0);
+	FScreensTask Task = GetMagicLeapScreensPlugin().UpdateWatchHistoryEntry(CurrentTask.WatchHistory[0]);
+	Task.EntryRequestDelegate = CurrentTask.EntryRequestDelegate;
+	CurrentTask = Task;
+	check(CurrentTask.WatchHistory.Num() > 0);
+	return CurrentTask.bSuccess;
+}
+
+bool FScreensRunnable::UpdateInfoEntry()
+{
+	FScreensTask Task = GetMagicLeapScreensPlugin().UpdateScreensTransform(CurrentTask.ScreenTransform);
+	Task.TransformRequestDelegate = CurrentTask.TransformRequestDelegate;
+	CurrentTask = Task;
+	return CurrentTask.bSuccess;
+}
+
+bool FScreensRunnable::GetWatchHistory()
+{
+	FScreensTask Task = GetMagicLeapScreensPlugin().GetWatchHistoryEntries();
+	Task.HistoryRequestDelegate = CurrentTask.HistoryRequestDelegate;
+	CurrentTask = Task;
+	return CurrentTask.bSuccess;
+}
+
diff --git a/Engine/Plugins/Lumin/MagicLeapScreens/Source/MagicLeapScreens/Private/MagicLeapScreensRunnable.h b/Engine/Plugins/Lumin/MagicLeapScreens/Source/MagicLeapScreens/Private/MagicLeapScreensRunnable.h
new file mode 100644
--- /dev/null
+++ b/Engine/Plugins/Lumin/MagicLeapScreens/Source/MagicLeapScreens/Private/MagicLeapScreensRunnable.h
@@ -0,0 +1,57 @@
+// Copyright 1998-2019 Epic Games, Inc. All Rights Reserved.
+
+#pragma once
+
+#include "MagicLeapRunnable.h"
+#include "MagicLeapScreensTypes.h"
+#include "Containers/Array.h"
+
+struct FScreensTask : public FMagicLeapTask
+{
+	enum class ETaskType : uint32
+	{
+		None,
+		GetHistory,
+		RemoveFromHistory,
+		AddToHistory,
+		UpdateHistoryEntry,
+		UpdateInfoEntry,
+	};
+
+	enum class ETaskRequestType : uint32
+	{
+		Request,
+		Response
+	};
+
+	ETaskType TaskType;
+	ETaskRequestType TaskRequestType;
+	TArray<FMagicLeapScreensWatchHistoryEntry> WatchHistory;
+	FMagicLeapScreenTransform ScreenTransform;
+
+	FMagicLeapScreensEntryRequestResultDelegate EntryRequestDelegate;
+	FMagicLeapScreensHistoryRequestResultDelegate HistoryRequestDelegate;
+	FMagicLeapScreenTransformRequestResultDelegate TransformRequestDelegate;
+
+	FScreensTask()
+		: TaskType(ETaskType::None)
+		, TaskRequestType(ETaskRequestType::Request)
+		, WatchHistory(TArray<FMagicLeapScreensWatchHistoryEntry>())
+	{
+	}
+};
+
+class FScreensRunnable : public FMagicLeapRunnable<FScreensTask>
+{
+public:
+	FScreensRunnable();
+	bool ProcessCurrentTask() override;
+	bool PeekCompletedTasks(FScreensTask& OutTask) const;
+	bool CompletedTaskQueueIsEmpty() const;
+private:
+	bool AddToHistory();
+	bool UpdateWatchHistoryEntry();
+	bool UpdateInfoEntry();
+	bool GetWatchHistory();
+};
+
diff --git a/Engine/Plugins/Lumin/MagicLeapScreens/Source/MagicLeapScreens/Public/IMagicLeapScreensPlugin.h b/Engine/Plugins/Lumin/MagicLeapScreens/Source/MagicLeapScreens/Public/IMagicLeapScreensPlugin.h
--- a/Engine/Plugins/Lumin/MagicLeapScreens/Source/MagicLeapScreens/Public/IMagicLeapScreensPlugin.h
+++ b/Engine/Plugins/Lumin/MagicLeapScreens/Source/MagicLeapScreens/Public/IMagicLeapScreensPlugin.h
@@ -33,13 +33,4 @@ public:
 	{
 		return FModuleManager::Get().IsModuleLoaded( "MagicLeapScreens" );
 	}
-
-	virtual bool Tick(float DeltaTime) = 0;
-
-	virtual bool IsEngineLoopInitComplete() const = 0;
-
-	virtual void OnEngineLoopInitComplete() = 0;
-
-	FTickerDelegate TickDelegate;
-	FDelegateHandle TickDelegateHandle;
 };
diff --git a/Engine/Plugins/Lumin/MagicLeapScreens/Source/MagicLeapScreens/Public/MagicLeapScreensFunctionLibrary.h b/Engine/Plugins/Lumin/MagicLeapScreens/Source/MagicLeapScreens/Public/MagicLeapScreensFunctionLibrary.h
--- a/Engine/Plugins/Lumin/MagicLeapScreens/Source/MagicLeapScreens/Public/MagicLeapScreensFunctionLibrary.h
+++ b/Engine/Plugins/Lumin/MagicLeapScreens/Source/MagicLeapScreens/Public/MagicLeapScreensFunctionLibrary.h
@@ -15,40 +15,49 @@ class MAGICLEAPSCREENS_API UMagicLeapScreensFunctionLibrary : public UBlueprintF
 public:
 
 	/**
-		Asynchronously requests all watch history entries.
+		Asynchronously requests all watch history entries that belong to this application.
 
 		@param[in] ResultDelegate Delegate that is called once the request has been resolved.
 	*/
 	UFUNCTION(BlueprintCallable, Category = "Screens Function Library|MagicLeap")
-	static void GetWatchHistoryAsync(const FScreensHistoryRequestResultDelegate& ResultDelegate);
+	static void GetWatchHistoryAsync(const FMagicLeapScreensHistoryRequestResultDelegate& ResultDelegate);
 
 	/**
 		Asynchronously requests to adds a new entry into the watch history.
 
 		@param[in] NewEntry Watch history entry to add.
-		@param[in] ResultDelegate Delegate that is called once the add request has been resolved. The resulting FScreensWatchHistoryEntry returned.
+		@param[in] ResultDelegate Delegate that is called once the add request has been resolved. The resulting FMagicLeapScreensWatchHistoryEntry returned.
 		on success contains the new ID of the added entry.
 	*/
 	UFUNCTION(BlueprintCallable, Category = "Screens Function Library|MagicLeap")
-	static void AddToWatchHistoryAsync(const FScreensWatchHistoryEntry& NewEntry, const FScreensEntryRequestResultDelegate& ResultDelegate);
+	static void AddToWatchHistoryAsync(const FMagicLeapScreensWatchHistoryEntry& NewEntry, const FMagicLeapScreensEntryRequestResultDelegate& ResultDelegate);
 
 	/**
-		Asynchronously requests to update an entry in the watch history
+		Asynchronously requests to update an entry in the watch history. Only entries associated with this application can be updated.
 
 		@param[in] UpdateEntry Watch history entry to update. The ID of this entry must be valid and in the watch history in order for the update to resolve successfully
 		@param[in] ResultDelegate Delegate that is called once the update request has been resolved.
 	*/
 	UFUNCTION(BlueprintCallable, Category = "Screens Function Library|MagicLeap")
-	static void UpdateWatchHistoryEntryAsync(const FScreensWatchHistoryEntry& UpdateEntry, const FScreensEntryRequestResultDelegate& ResultDelegate);
+	static void UpdateWatchHistoryEntryAsync(const FMagicLeapScreensWatchHistoryEntry& UpdateEntry, const FMagicLeapScreensEntryRequestResultDelegate& ResultDelegate);
 
 	/**
-		Removes an entry from the watch history that corresponse with the ID passed.
+		Asynchronously requests to update a screen transform
+
+		@param[in] UpdateTransform Screen transform to update. The ID of this entry must be valid in order for the update to resolve successfully
+		@param[in] ResultDelegate Delegate that is called once the update request has been resolved.
+	*/
+	UFUNCTION(BlueprintCallable, Category = "Screens Function Library|MagicLeap")
+	static void UpdateScreenTransformAsync(const FMagicLeapScreenTransform& UpdateTransform, const FMagicLeapScreenTransformRequestResultDelegate& ResultDelegate);
+
+	/**
+		Removes an entry from the watch history that corresponds with the ID passed. Only entries associated with this application can be removed.
 
 		@param[in] ID ID of the watch history entry to remove.
 		@return True when the entry is sucessfully removed.
 	*/
 	UFUNCTION(BlueprintCallable, Category = "Screens Function Library|MagicLeap")
-	static bool RemoveWatchHistoryEntry(const FScreenID& ID);
+	static bool RemoveWatchHistoryEntry(const FGuid& ID);
 
 	/**
 		Removes all watch history entries.
@@ -61,21 +70,36 @@ public:
 	/**
 		Gets the transform for all watch history entries.
 
-		@param[out] ScreenTransforms Array of transforms that corresponse to all current watch history entries.
+		@param[out] ScreenTransforms Array of transforms that corresponds to all current watch history entries.
 		@return True when the request for all transforms succeeds 
 	*/
 	UFUNCTION(BlueprintCallable, Category = "Screens Function Library|MagicLeap")
-	static bool GetScreensTransforms(TArray<FScreenTransform>& ScreenTransforms);
+	static bool GetScreensTransforms(TArray<FMagicLeapScreenTransform>& ScreenTransforms);
+
+	/**
+		Gets a single screen transform.
+
+		@param[out] ScreenTransform Transform that corresponds to a watch history entry with the ID passed via command line as 'screenId'.
+		@return True when the request for the transform succeeds
+	*/
+	UFUNCTION(BlueprintCallable, Category = "Screens Function Library|MagicLeap")
+	static bool GetScreenTransform(FMagicLeapScreenTransform& ScreenTransform);
 
 	/**
 		Delegate used to relay the results of GetWatchHistory.
 	*/
 	UPROPERTY()
-	FScreensHistoryRequestResultDelegate HistoryResultDelegate;
+	FMagicLeapScreensHistoryRequestResultDelegate HistoryResultDelegate;
 
 	/**
 		Delegate used to relay the results of both AddToWatchHistory and UpdateWatchHistoryEntry.
 	*/
 	UPROPERTY()
-	FScreensEntryRequestResultDelegate EntryResultDelegate;
+	FMagicLeapScreensEntryRequestResultDelegate EntryResultDelegate;
+
+	/**
+		Delegate used to relay the results of UpdateScreenTransformAsync.
+	*/
+	UPROPERTY()
+	FMagicLeapScreenTransformRequestResultDelegate ScreenTransformResultDelegate;
 };
diff --git a/Engine/Plugins/Lumin/MagicLeapScreens/Source/MagicLeapScreens/Public/MagicLeapScreensTypes.h b/Engine/Plugins/Lumin/MagicLeapScreens/Source/MagicLeapScreens/Public/MagicLeapScreensTypes.h
--- a/Engine/Plugins/Lumin/MagicLeapScreens/Source/MagicLeapScreens/Public/MagicLeapScreensTypes.h
+++ b/Engine/Plugins/Lumin/MagicLeapScreens/Source/MagicLeapScreens/Public/MagicLeapScreensTypes.h
@@ -5,30 +5,17 @@
 #include "Engine/Engine.h"
 #include "MagicLeapScreensTypes.generated.h"
 
-/** 
-  ID for a Screens Watch History Entry.
-
-  Save this off when you add a new watch history and use the same to update or delete that same entry.
- */
-USTRUCT(BlueprintType)
-struct FScreenID
-{
-	GENERATED_BODY()
-
-public:
-	int64 ID;
-};
 
 /** Channel watch history, may be displayed in the Screens Launcher application. */
 USTRUCT(BlueprintType)
-struct FScreensWatchHistoryEntry
+struct FMagicLeapScreensWatchHistoryEntry
 {
 	GENERATED_BODY()
 
 public:
 	/** Entry Identifier. Must be used to update and delete a given entry. */
 	UPROPERTY(VisibleAnywhere, BlueprintReadOnly, Category = "Screens|MagicLeap")
-	FScreenID ID;
+	FGuid ID;
 
 	/** Title of the media for which this entry is created. */
 	UPROPERTY(EditAnywhere, BlueprintReadWrite, Category = "Screens|MagicLeap")
@@ -61,11 +48,15 @@ public:
   This will be received from the Screens Launcher api, based on the previous screens spawned by user.
  */
 USTRUCT(BlueprintType)
-struct FScreenTransform
+struct FMagicLeapScreenTransform
 {
 	GENERATED_BODY()
 
 public:
+	/** Entry identifier. Must be used to update a given entry. */
+	UPROPERTY(VisibleAnywhere, BlueprintReadOnly, Category = "Screens|MagicLeap")
+	FGuid ID;
+
 	/** Position of the screen in Unreal's world space. */
 	UPROPERTY(EditAnywhere, BlueprintReadWrite, Category = "Screens|MagicLeap")
 	FVector ScreenPosition;
@@ -77,6 +68,12 @@ public:
 	/** Dimensions of the screen in Unreal Units. The dimensions are axis-aligned with the orientation. */
 	UPROPERTY(EditAnywhere, BlueprintReadWrite, Category = "Screens|MagicLeap")
 	FVector ScreenDimensions;
+
+	/** Screen scale, used to scale the screens dimensions */
+	FVector ScreenScale3D;
+
+	/** Version number */
+	int32 VersionNumber;
 };
 
 /**
@@ -86,7 +83,7 @@ public:
 	@param[out] bSuccess True when the request is successful
 	@param[out] WatchHistoryEntry Resulting watch history entry for which the operation was performed on.
 */
-DECLARE_DYNAMIC_DELEGATE_TwoParams(FScreensEntryRequestResultDelegate, const bool, bSuccess, const FScreensWatchHistoryEntry&, WatchHistoryEntry);
+DECLARE_DYNAMIC_DELEGATE_TwoParams(FMagicLeapScreensEntryRequestResultDelegate, const bool, bSuccess, const FMagicLeapScreensWatchHistoryEntry&, WatchHistoryEntry);
 
 /**
 	Delegate used to relay the result of getting the entire watch history.
@@ -94,4 +91,12 @@ DECLARE_DYNAMIC_DELEGATE_TwoParams(FScreensEntryRequestResultDelegate, const boo
 	@param[out] bSuccess True when the request is successful.
 	@param[out] WatchHistoryEntry Resulting array of watch histories returned by the operation.
 */
-DECLARE_DYNAMIC_DELEGATE_TwoParams(FScreensHistoryRequestResultDelegate, const bool, bSuccess, const TArray<FScreensWatchHistoryEntry>&, WatchHistoryEntries);
\ No newline at end of file
+DECLARE_DYNAMIC_DELEGATE_TwoParams(FMagicLeapScreensHistoryRequestResultDelegate, const bool, bSuccess, const TArray<FMagicLeapScreensWatchHistoryEntry>&, WatchHistoryEntries);
+
+/**
+	Delegate used to relay the result of updating a screen's transform.
+
+	@param[out] bSuccess True when the request is successful.
+*/
+DECLARE_DYNAMIC_DELEGATE_OneParam(FMagicLeapScreenTransformRequestResultDelegate, const bool, bSuccess);
+
diff --git a/Engine/Plugins/Lumin/MagicLeapSharedFile/Config/BaseMagicLeapSharedFile.ini b/Engine/Plugins/Lumin/MagicLeapSharedFile/Config/BaseMagicLeapSharedFile.ini
new file mode 100644
--- /dev/null
+++ b/Engine/Plugins/Lumin/MagicLeapSharedFile/Config/BaseMagicLeapSharedFile.ini
@@ -0,0 +1,2 @@
+[CoreRedirects]
++ClassRedirects=(OldName="/Script/MagicLeapSharedFile.FilesPickedResultDelegate", NewName="/Script/MagicLeapSharedFile.MagicLeapFilesPickedResultDelegate")
diff --git a/Engine/Plugins/Lumin/MagicLeapSharedFile/MagicLeapSharedFile.uplugin b/Engine/Plugins/Lumin/MagicLeapSharedFile/MagicLeapSharedFile.uplugin
new file mode 100644
--- /dev/null
+++ b/Engine/Plugins/Lumin/MagicLeapSharedFile/MagicLeapSharedFile.uplugin
@@ -0,0 +1,42 @@
+{
+	"FileVersion" : 3,
+	"Version" : 1,
+	"VersionName": "1.0",
+	"FriendlyName" : "Magic Leap Shared File",
+	"CreatedBy" : "Magic Leap, Inc.",
+	"CreatedByURL" : "http://magicleap.com",
+	"Description" : "Support for the Magic Leap SharedFile API.",
+	"Category" : "Lumin",
+	"EnabledByDefault" : false,
+	"DocsURL" : "",
+	"MarketplaceURL" : "",
+	"SupportURL" : "",
+	"CanContainContent" : false,
+	"IsBetaVersion" : false,
+	"Installed" : false,
+	"SupportedTargetPlatforms": [
+		"Lumin",
+		"Win64",
+		"Mac"
+	],
+	"Modules" :
+	[
+		{
+			"Name" : "MagicLeapSharedFile",
+			"Type" : "Runtime",
+			"WhitelistPlatforms" :
+			[
+				"Win64",
+				"Mac",
+				"Lumin"
+			]
+		}
+	],
+	"Plugins":
+	[
+		{
+			"Name": "MLSDK",
+			"Enabled": true
+		}
+	]
+}
diff --git a/Engine/Plugins/Lumin/MagicLeapSharedFile/Source/MagicLeapSharedFile.Build.cs b/Engine/Plugins/Lumin/MagicLeapSharedFile/Source/MagicLeapSharedFile.Build.cs
new file mode 100644
--- /dev/null
+++ b/Engine/Plugins/Lumin/MagicLeapSharedFile/Source/MagicLeapSharedFile.Build.cs
@@ -0,0 +1,31 @@
+// Copyright 1998-2019 Epic Games, Inc. All Rights Reserved.
+
+using System.IO;
+
+namespace UnrealBuildTool.Rules
+{
+	public class MagicLeapSharedFile : ModuleRules
+	{
+		public MagicLeapSharedFile( ReadOnlyTargetRules Target ) : base(Target)
+		{
+			PublicIncludePaths.Add(Path.Combine(ModuleDirectory, "Public"));
+
+			PublicDependencyModuleNames.AddRange(
+				new string[]
+				{
+					"Core",
+					"CoreUObject",
+					"Engine"
+				}
+			);
+
+			PrivateDependencyModuleNames.AddRange(
+				new string[]
+				{
+					"LuminRuntimeSettings",
+					"MLSDK"
+				}
+			);
+		}
+	}
+}
diff --git a/Engine/Plugins/Lumin/MagicLeapSharedFile/Source/Private/MagicLeapSharedFileFunctionLibrary.cpp b/Engine/Plugins/Lumin/MagicLeapSharedFile/Source/Private/MagicLeapSharedFileFunctionLibrary.cpp
new file mode 100644
--- /dev/null
+++ b/Engine/Plugins/Lumin/MagicLeapSharedFile/Source/Private/MagicLeapSharedFileFunctionLibrary.cpp
@@ -0,0 +1,14 @@
+// Copyright 1998-2019 Epic Games, Inc. All Rights Reserved.
+
+#include "MagicLeapSharedFileFunctionLibrary.h"
+#include "MagicLeapSharedFilePlugin.h"
+
+bool UMagicLeapSharedFileFunctionLibrary::SharedFileListAccessibleFiles(TArray<FString>& OutSharedFileList)
+{
+	return GetMagicLeapSharedFilePlugin().SharedFileListAccessibleFiles(OutSharedFileList);
+}
+
+bool UMagicLeapSharedFileFunctionLibrary::SharedFilePickAsync(const FMagicLeapFilesPickedResultDelegate& InResultDelegate)
+{
+	return GetMagicLeapSharedFilePlugin().SharedFilePickAsync(InResultDelegate);
+}
diff --git a/Engine/Plugins/Lumin/MagicLeapSharedFile/Source/Private/MagicLeapSharedFilePlugin.cpp b/Engine/Plugins/Lumin/MagicLeapSharedFile/Source/Private/MagicLeapSharedFilePlugin.cpp
new file mode 100644
--- /dev/null
+++ b/Engine/Plugins/Lumin/MagicLeapSharedFile/Source/Private/MagicLeapSharedFilePlugin.cpp
@@ -0,0 +1,167 @@
+// Copyright 1998-2019 Epic Games, Inc. All Rights Reserved.
+
+#include "MagicLeapSharedFilePlugin.h"
+#include "Lumin/CAPIShims/LuminAPISharedFile.h"
+#include "Lumin/CAPIShims/LuminAPIFileInfo.h"
+
+#if PLATFORM_LUMIN
+#include "Lumin/LuminPlatformFile.h"
+#endif // PLATFORM_LUMIN
+
+DEFINE_LOG_CATEGORY(LogMagicLeapSharedFile);
+
+
+#if WITH_MLSDK
+bool GetFilenamesFromSharedFileList_Impl(MLSharedFileList* SharedFileList, TArray<FString>& OutSharedFileList)
+{
+	MLResult Result = MLResult_UnspecifiedFailure;
+	OutSharedFileList.Empty();
+
+	if (SharedFileList != nullptr)
+	{
+		MLHandle ListLength = 0;
+		Result = MLSharedFileGetListLength(SharedFileList, &ListLength);
+		UE_CLOG(MLResult_Ok != Result, LogMagicLeapSharedFile, Error, TEXT("MLSharedFileGetListLength failed with error '%s'"), UTF8_TO_TCHAR(MLSharedFileGetResultString(Result)));
+		if (Result == MLResult_Ok)
+		{
+			for (MLHandle i = 0; i < ListLength; ++i)
+			{
+				MLFileInfo* FileInfo = nullptr;
+				Result = MLSharedFileGetMLFileInfoByIndex(SharedFileList, i, &FileInfo);
+				UE_CLOG(MLResult_Ok != Result, LogMagicLeapSharedFile, Error, TEXT("MLSharedFileGetMLFileInfoByIndex failed with error '%s'"), UTF8_TO_TCHAR(MLSharedFileGetResultString(Result)));
+				if (Result == MLResult_Ok && FileInfo != nullptr)
+				{
+					const char* Filename = nullptr;
+					Result = MLFileInfoGetFileName(FileInfo, &Filename);
+					UE_CLOG(MLResult_Ok != Result, LogMagicLeapSharedFile, Error, TEXT("MLFileInfoGetFileName failed with error '%s'"), UTF8_TO_TCHAR(MLSharedFileGetResultString(Result)));
+					if (Result == MLResult_Ok && Filename != nullptr)
+					{
+						OutSharedFileList.Add(UTF8_TO_TCHAR(Filename));
+					}
+				}
+			}
+		}
+
+		Result = MLSharedFileListRelease(&SharedFileList);
+		UE_CLOG(MLResult_Ok != Result, LogMagicLeapSharedFile, Error, TEXT("MLSharedFileListRelease failed with error '%s'"), UTF8_TO_TCHAR(MLSharedFileGetResultString(Result)));
+	}
+
+	return MLResult_Ok == Result;
+}
+
+void onFilesPicked(MLSharedFileList *list, void *context)
+{
+	check(context != nullptr);
+	FMagicLeapSharedFilePlugin* Plugin = reinterpret_cast<FMagicLeapSharedFilePlugin*>(context);
+	Plugin->GetFileNamesFromSharedFileList(reinterpret_cast<void*>(list));
+}
+#endif // WITH_MLSDK
+
+void FMagicLeapSharedFilePlugin::StartupModule()
+{
+	IMagicLeapSharedFilePlugin::StartupModule();
+
+	TickDelegate = FTickerDelegate::CreateRaw(this, &FMagicLeapSharedFilePlugin::Tick);
+}
+
+void FMagicLeapSharedFilePlugin::ShutdownModule()
+{
+	FTicker::GetCoreTicker().RemoveTicker(TickDelegateHandle);
+	IMagicLeapSharedFilePlugin::ShutdownModule();
+}
+
+IFileHandle* FMagicLeapSharedFilePlugin::SharedFileRead(const FString& FileName)
+{
+#if PLATFORM_LUMIN
+	IPlatformFile& PlatformFile = IPlatformFile::GetPlatformPhysical();
+	FLuminPlatformFile* LuminPlatformFile = static_cast<FLuminPlatformFile*>(&PlatformFile);
+	return LuminPlatformFile->SharedFileOpenRead(*FileName);
+#else
+	return nullptr;
+#endif // PLATFORM_LUMIN
+}
+
+IFileHandle* FMagicLeapSharedFilePlugin::SharedFileWrite(const FString& FileName)
+{
+#if PLATFORM_LUMIN
+	IPlatformFile& PlatformFile = IPlatformFile::GetPlatformPhysical();
+	FLuminPlatformFile* LuminPlatformFile = static_cast<FLuminPlatformFile*>(&PlatformFile);
+	return LuminPlatformFile->SharedFileOpenWrite(*FileName);
+#else
+	return nullptr;
+#endif // PLATFORM_LUMIN
+}
+
+bool FMagicLeapSharedFilePlugin::SharedFileListAccessibleFiles(TArray<FString>& OutSharedFileList)
+{
+	OutSharedFileList.Empty();
+#if WITH_MLSDK
+	MLSharedFileList* SharedFileList = nullptr;
+	MLResult Result = MLSharedFileListAccessibleFiles(&SharedFileList);
+	UE_CLOG(MLResult_Ok != Result, LogMagicLeapSharedFile, Error, TEXT("MLSharedFileListAccessibleFiles failed with error '%s'"), UTF8_TO_TCHAR(MLSharedFileGetResultString(Result)));
+	if (Result == MLResult_Ok && SharedFileList != nullptr)
+	{
+		return GetFilenamesFromSharedFileList_Impl(SharedFileList, OutSharedFileList);
+	}
+#endif // WITH_MLSDK
+
+	return false;
+}
+
+bool FMagicLeapSharedFilePlugin::SharedFilePickAsync(const FMagicLeapFilesPickedResultDelegate& InResultDelegate)
+{
+#if WITH_MLSDK
+	ResultDelegate = InResultDelegate;
+
+	{
+		FScopeLock Lock(&Mutex);
+		bWaitingForDelegateResult = true;
+	}
+
+	TickDelegateHandle = FTicker::GetCoreTicker().AddTicker(TickDelegate);
+
+	MLResult Result = MLSharedFilePick(onFilesPicked, this);
+	if (MLResult_Ok != Result)
+	{
+		FScopeLock Lock(&Mutex);
+		bWaitingForDelegateResult = true;
+		FTicker::GetCoreTicker().RemoveTicker(TickDelegateHandle);
+		UE_LOG(LogMagicLeapSharedFile, Error, TEXT("MLSharedFilePick failed with error '%s'"), UTF8_TO_TCHAR(MLSharedFileGetResultString(Result)));
+	}
+
+	return Result == MLResult_Ok;
+#else
+	return false;
+#endif // WITH_MLSDK
+}
+
+bool FMagicLeapSharedFilePlugin::Tick(float DeltaTime)
+{
+	bool bWaitingForDelegateResult_Cached;
+	FMagicLeapFilesPickedResultDelegate ResultDelegate_Cached;
+
+	{
+		FScopeLock Lock(&Mutex);
+		bWaitingForDelegateResult_Cached = bWaitingForDelegateResult;
+		ResultDelegate_Cached = ResultDelegate;
+	}
+	
+	if (!bWaitingForDelegateResult_Cached)
+	{
+		FTicker::GetCoreTicker().RemoveTicker(TickDelegateHandle);
+		ResultDelegate_Cached.ExecuteIfBound(PickedFileList);
+	}
+
+	return true;
+}
+
+void FMagicLeapSharedFilePlugin::GetFileNamesFromSharedFileList(void* SharedFiles)
+{
+#if WITH_MLSDK
+	FScopeLock Lock(&Mutex);
+	GetFilenamesFromSharedFileList_Impl(static_cast<MLSharedFileList*>(SharedFiles), PickedFileList);
+	bWaitingForDelegateResult = false;
+#endif // WITH_MLSDK
+}
+
+IMPLEMENT_MODULE(FMagicLeapSharedFilePlugin, MagicLeapSharedFile);
diff --git a/Engine/Plugins/Lumin/MagicLeapSharedFile/Source/Private/MagicLeapSharedFilePlugin.h b/Engine/Plugins/Lumin/MagicLeapSharedFile/Source/Private/MagicLeapSharedFilePlugin.h
new file mode 100644
--- /dev/null
+++ b/Engine/Plugins/Lumin/MagicLeapSharedFile/Source/Private/MagicLeapSharedFilePlugin.h
@@ -0,0 +1,40 @@
+// Copyright 1998-2019 Epic Games, Inc. All Rights Reserved.
+
+#pragma once
+
+#include "CoreMinimal.h"
+#include "IMagicLeapSharedFilePlugin.h"
+#include "MagicLeapPluginUtil.h"
+#include "Containers/Ticker.h"
+
+DECLARE_LOG_CATEGORY_EXTERN(LogMagicLeapSharedFile, Verbose, All);
+
+class FMagicLeapSharedFilePlugin : public IMagicLeapSharedFilePlugin
+{
+public:
+	void StartupModule() override;
+	void ShutdownModule() override;
+
+	virtual IFileHandle* SharedFileRead(const FString& FileName) override;
+	virtual IFileHandle* SharedFileWrite(const FString& FileName) override;
+	virtual bool SharedFileListAccessibleFiles(TArray<FString>& OutSharedFileList) override;
+	virtual bool SharedFilePickAsync(const FMagicLeapFilesPickedResultDelegate& InResultDelegate) override;
+
+	bool Tick(float DeltaTime);
+	void GetFileNamesFromSharedFileList(void* SharedFiles);
+
+private:
+	FMagicLeapAPISetup APISetup;
+	FTickerDelegate TickDelegate;
+	FDelegateHandle TickDelegateHandle;
+
+	FCriticalSection Mutex;
+	FMagicLeapFilesPickedResultDelegate ResultDelegate;
+	bool bWaitingForDelegateResult;
+	TArray<FString> PickedFileList;
+};
+
+inline FMagicLeapSharedFilePlugin& GetMagicLeapSharedFilePlugin()
+{
+	return FModuleManager::Get().GetModuleChecked<FMagicLeapSharedFilePlugin>("MagicLeapSharedFile");
+}
diff --git a/Engine/Plugins/Lumin/MagicLeapSharedFile/Source/Public/IMagicLeapSharedFilePlugin.h b/Engine/Plugins/Lumin/MagicLeapSharedFile/Source/Public/IMagicLeapSharedFilePlugin.h
new file mode 100644
--- /dev/null
+++ b/Engine/Plugins/Lumin/MagicLeapSharedFile/Source/Public/IMagicLeapSharedFilePlugin.h
@@ -0,0 +1,67 @@
+// Copyright 1998-2019 Epic Games, Inc. All Rights Reserved.
+
+#pragma once
+
+#include "Modules/ModuleManager.h"
+#include "MagicLeapSharedFileTypes.h"
+
+/**
+ * The public interface to this module.  In most cases, this interface is only public to sibling modules
+ * within this plugin.
+ */
+class MAGICLEAPSHAREDFILE_API IMagicLeapSharedFilePlugin : public IModuleInterface
+{
+public:
+
+	/**
+	 * Singleton-like access to this module's interface.  This is just for convenience!
+	 * Beware of calling this during the shutdown phase, though.  Your module might have been unloaded already.
+	 *
+	 * @return Returns singleton instance, loading the module on demand if needed
+	 */
+	static inline IMagicLeapSharedFilePlugin& Get()
+	{
+		return FModuleManager::LoadModuleChecked<IMagicLeapSharedFilePlugin>("MagicLeapSharedFile");
+	}
+
+	/**
+	 * Checks to see if this module is loaded and ready.  It is only valid to call Get() if IsAvailable() returns true.
+	 *
+	 * @return True if the module is loaded and ready to use
+	 */
+	static inline bool IsAvailable()
+	{
+		return FModuleManager::Get().IsModuleLoaded("MagicLeapSharedFile");
+	}
+
+	/**
+		Return a IFileHandle pointer to read the user shared file. Caller owns the IFileHandle pointer and should delete it after use.
+		@param FileName Name of the shared file to read.
+		@return IFileHandle pointer to read the file. If the application does not have user permission to access the file, nullptr will be returned.
+	*/
+	virtual IFileHandle* SharedFileRead(const FString& FileName) = 0;
+
+	/**
+		Return a IFileHandle pointer to write the user shared file. Caller owns the IFileHandle pointer and should delete it after use.
+		@param FileName Name of the shared file to write to. Needs to just be a file name, cannot be a path.
+		@return IFileHandle pointer to read the file. If the application does not have user permission to access the file, nullptr will be returned.
+	*/
+	virtual IFileHandle* SharedFileWrite(const FString& FileName) = 0;
+
+	/**
+		Get the names of the files that the application has access to.
+		The application can then use the file names and read them with the IMagicLeapSharedFilePlugin::SharedFileRead() function.
+		@param OutSharedFileList Output param containing list of file names this app has acces to.
+		@return true if function call succeeded and output param is valid, false otherwise
+	*/
+	virtual bool SharedFileListAccessibleFiles(TArray<FString>& OutSharedFileList) = 0;
+
+	/**
+		Let the app get access to the user's shared files from the common storage location.  This API will pop up a System UI
+		dialogue box with a file picker through which the users can pick files they wants to let the app have access to.  The list
+		of selected file names will be returned to the app via the delegate.
+		@param InResultDelegate Delegate to be called when the user finishes picking files.
+		@return true if call to invoke the file picker succeeded, false otherwise
+	*/
+	virtual bool SharedFilePickAsync(const FMagicLeapFilesPickedResultDelegate& InResultDelegate) = 0;
+};
diff --git a/Engine/Plugins/Lumin/MagicLeapSharedFile/Source/Public/MagicLeapSharedFileFunctionLibrary.h b/Engine/Plugins/Lumin/MagicLeapSharedFile/Source/Public/MagicLeapSharedFileFunctionLibrary.h
new file mode 100644
--- /dev/null
+++ b/Engine/Plugins/Lumin/MagicLeapSharedFile/Source/Public/MagicLeapSharedFileFunctionLibrary.h
@@ -0,0 +1,33 @@
+// Copyright 1998-2019 Epic Games, Inc. All Rights Reserved.
+#pragma once
+
+#include "Engine/Engine.h"
+#include "Kismet/BlueprintFunctionLibrary.h"
+#include "MagicLeapSharedFileTypes.h"
+#include "MagicLeapSharedFileFunctionLibrary.generated.h"
+
+UCLASS(ClassGroup = MagicLeap)
+class MAGICLEAPSHAREDFILE_API UMagicLeapSharedFileFunctionLibrary : public UBlueprintFunctionLibrary
+{
+	GENERATED_BODY()
+
+public:
+	/**
+		Get the names of the files that the application has access to.
+		The application can then use the file names and read them with the IMagicLeapSharedFilePlugin::SharedFileRead() function.
+		@param OutSharedFileList Output param containing list of file names this app has acces to.
+		@return true if function call succeeded and output param is valid, false otherwise
+	*/
+	UFUNCTION(BlueprintCallable, Category = "SharedFile Function Library | MagicLeap")
+	static bool SharedFileListAccessibleFiles(TArray<FString>& OutSharedFileList);
+
+	/**
+		Let the app get access to the user's shared files from the common storage location.  This API will pop up a System UI
+		dialogue box with a file picker through which the users can pick files they wants to let the app have access to.  The list
+		of selected file names will be returned to the app via the delegate.
+		@param InResultDelegate Delegate to be called when the user finishes picking files.
+		@return true if call to invoke the file picker succeeded, false otherwise
+	*/
+	UFUNCTION(BlueprintCallable, Category = "SharedFile Function Library | MagicLeap")
+	static bool SharedFilePickAsync(const FMagicLeapFilesPickedResultDelegate& InResultDelegate);
+};
diff --git a/Engine/Plugins/Lumin/MagicLeapSharedFile/Source/Public/MagicLeapSharedFileTypes.h b/Engine/Plugins/Lumin/MagicLeapSharedFile/Source/Public/MagicLeapSharedFileTypes.h
new file mode 100644
--- /dev/null
+++ b/Engine/Plugins/Lumin/MagicLeapSharedFile/Source/Public/MagicLeapSharedFileTypes.h
@@ -0,0 +1,18 @@
+// Copyright 1998-2019 Epic Games, Inc. All Rights Reserved.
+#pragma once
+
+#include "Engine/Engine.h"
+#include "MagicLeapSharedFileTypes.generated.h"
+
+USTRUCT()
+struct FDummyMLSharedFileStruct
+{
+    GENERATED_BODY()
+};
+
+/**
+ * Delegate used to convey the result of a file pick operation. 
+ * @param List of file names the user picked to provide access to this app. The file names can then be used to read the file using the IMagicLeapSharedFilePlugin::SharedFileRead() function.
+ */
+DECLARE_DYNAMIC_DELEGATE_OneParam(FMagicLeapFilesPickedResultDelegate, const TArray<FString>&, PickedFiles);
+DECLARE_DYNAMIC_MULTICAST_DELEGATE_OneParam(FMagicLeapFilesPickedResultDelegateMulti, const TArray<FString>&, PickedFiles);
diff --git a/Engine/Plugins/Lumin/MagicLeapTablet/Config/BaseMagicLeapTablet.ini b/Engine/Plugins/Lumin/MagicLeapTablet/Config/BaseMagicLeapTablet.ini
new file mode 100644
--- /dev/null
+++ b/Engine/Plugins/Lumin/MagicLeapTablet/Config/BaseMagicLeapTablet.ini
@@ -0,0 +1,24 @@
+[CoreRedirects]
++EnumRedirects=(OldName="/Script/MagicLeapTablet.EMLInputTabletDeviceType", NewName="/Script/MagicLeapTablet.EMagicLeapInputTabletDeviceType")
++EnumRedirects=(OldName="/Script/MagicLeapTablet.EMLInputTabletDeviceToolType", NewName="/Script/MagicLeapTablet.EMagicLeapInputTabletDeviceToolType")
++EnumRedirects=(OldName="/Script/MagicLeapTablet.EMLInputTabletDeviceButton", NewName="/Script/MagicLeapTablet.EMagicLeapInputTabletDeviceButton")
+
++StructRedirects=(OldName="/Script/MagicLeapTablet.MLInputTabletDeviceState", NewName="/Script/MagicLeapTablet.MagicLeapInputTabletDeviceState")
+
++ClassRedirects=(OldName="/Script/MagicLeapTablet.TabletOnConnectedDelegate", NewName="/Script/MagicLeapTablet.MagicLeapTabletOnConnectedDelegate")
++ClassRedirects=(OldName="/Script/MagicLeapTablet.TabletOnConnectedDelegateMulti", NewName="/Script/MagicLeapTablet.MagicLeapTabletOnConnectedDelegateMulti")
+
++ClassRedirects=(OldName="/Script/MagicLeapTablet.TabletOnDisconnectedDelegate", NewName="/Script/MagicLeapTablet.MagicLeapTabletOnDisconnectedDelegate")
++ClassRedirects=(OldName="/Script/MagicLeapTablet.TabletOnDisconnectedDelegateMulti", NewName="/Script/MagicLeapTablet.MagicLeapTabletOnDisconnectedDelegateMulti")
+
++ClassRedirects=(OldName="/Script/MagicLeapTablet.TabletOnPenTouchDelegate", NewName="/Script/MagicLeapTablet.MagicLeapTabletOnPenTouchDelegate")
++ClassRedirects=(OldName="/Script/MagicLeapTablet.TabletOnPenTouchDelegateMulti", NewName="/Script/MagicLeapTablet.MagicLeapTabletOnPenTouchDelegateMulti")
+
++ClassRedirects=(OldName="/Script/MagicLeapTablet.TabletOnRingTouchDelegate", NewName="/Script/MagicLeapTablet.MagicLeapTabletOnRingTouchDelegate")
++ClassRedirects=(OldName="/Script/MagicLeapTablet.TabletOnRingTouchDelegateMulti", NewName="/Script/MagicLeapTablet.MagicLeapTabletOnRingTouchDelegateMulti")
+
++ClassRedirects=(OldName="/Script/MagicLeapTablet.TabletOnButtonDownDelegate", NewName="/Script/MagicLeapTablet.MagicLeapTabletOnButtonDownDelegate")
++ClassRedirects=(OldName="/Script/MagicLeapTablet.TabletOnButtonDownDelegateMulti", NewName="/Script/MagicLeapTablet.MagicLeapTabletOnButtonDownDelegateMulti")
+
++ClassRedirects=(OldName="/Script/MagicLeapTablet.TabletOnButtonUpDelegate", NewName="/Script/MagicLeapTablet.MagicLeapTabletOnButtonUpDelegate")
++ClassRedirects=(OldName="/Script/MagicLeapTablet.TabletOnButtonUpDelegateMulti", NewName="/Script/MagicLeapTablet.MagicLeapTabletOnButtonUpDelegateMulti")
diff --git a/Engine/Plugins/Lumin/MagicLeapTablet/MagicLeapTablet.uplugin b/Engine/Plugins/Lumin/MagicLeapTablet/MagicLeapTablet.uplugin
new file mode 100644
--- /dev/null
+++ b/Engine/Plugins/Lumin/MagicLeapTablet/MagicLeapTablet.uplugin
@@ -0,0 +1,44 @@
+{
+	"FileVersion": 3,
+	"Version": 1,
+	"VersionName": "1.0",
+	"FriendlyName": "Magic Leap Tablet",
+	"CreatedBy": "Magic Leap, Inc.",
+	"CreatedByURL": "http://magicleap.com",
+	"Description": "Support for the Magic Leap Tablet API.",
+	"Category": "Lumin",
+	"EnabledByDefault": false,
+	"DocsURL": "",
+	"MarketplaceURL": "",
+	"SupportURL": "",
+	"CanContainContent": false,
+	"IsBetaVersion": false,
+	"Installed": false,
+	"SupportedTargetPlatforms": [
+		"Lumin",
+		"Win64",
+		"Mac"
+	],
+	"Modules": [
+		{
+			"Name": "MagicLeapTablet",
+			"Type": "Runtime",
+			"LoadingPhase": "PostConfigInit",
+			"WhitelistPlatforms": [
+				"Win64",
+				"Mac",
+				"Lumin"
+			]
+		}
+	],
+	"Plugins": [
+		{
+			"Name": "MLSDK",
+			"Enabled": true
+		},
+		{
+			"Name": "MagicLeap",
+			"Enabled": true
+		}
+	]
+}
diff --git a/Engine/Plugins/Lumin/MagicLeapTablet/Source/MagicLeapTablet.Build.cs b/Engine/Plugins/Lumin/MagicLeapTablet/Source/MagicLeapTablet.Build.cs
new file mode 100644
--- /dev/null
+++ b/Engine/Plugins/Lumin/MagicLeapTablet/Source/MagicLeapTablet.Build.cs
@@ -0,0 +1,29 @@
+// Copyright 1998-2019 Epic Games, Inc. All Rights Reserved.
+
+using System.IO;
+
+namespace UnrealBuildTool.Rules
+{
+	public class MagicLeapTablet : ModuleRules
+	{
+		public MagicLeapTablet( ReadOnlyTargetRules Target ) : base(Target)
+		{
+			PublicIncludePaths.Add(Path.Combine(ModuleDirectory, "Public"));
+
+            PublicDependencyModuleNames.Add("InputDevice");
+
+            PrivateDependencyModuleNames.AddRange(
+                new string[]
+                {
+                    "Core",
+                    "CoreUObject",
+                    "ApplicationCore",
+                    "Engine",
+                    "InputCore",
+                    "MagicLeap",
+                    "MLSDK"
+                }
+            );
+        }
+	}
+}
diff --git a/Engine/Plugins/Lumin/MagicLeapTablet/Source/Private/MagicLeapTabletComponent.cpp b/Engine/Plugins/Lumin/MagicLeapTablet/Source/Private/MagicLeapTabletComponent.cpp
new file mode 100644
--- /dev/null
+++ b/Engine/Plugins/Lumin/MagicLeapTablet/Source/Private/MagicLeapTabletComponent.cpp
@@ -0,0 +1,21 @@
+// Copyright 1998-2019 Epic Games, Inc. All Rights Reserved.
+
+#include "MagicLeapTabletComponent.h"
+#include "MagicLeapTabletPlugin.h"
+
+UMagicLeapTabletComponent::UMagicLeapTabletComponent()
+{
+	bAutoActivate = true;
+}
+
+void UMagicLeapTabletComponent::BeginPlay()
+{
+	Super::BeginPlay();
+
+	GetMagicLeapTabletPlugin().SetConnectedDelegate(OnConnected);
+	GetMagicLeapTabletPlugin().SetDisconnectedDelegate(OnDisconnected);
+	GetMagicLeapTabletPlugin().SetPenTouchDelegate(OnPenTouch);
+	GetMagicLeapTabletPlugin().SetRingTouchDelegate(OnRingTouch);
+	GetMagicLeapTabletPlugin().SetButtonDownDelegate(OnButtonDown);
+	GetMagicLeapTabletPlugin().SetButtonUpDelegate(OnButtonUp);
+}
diff --git a/Engine/Plugins/Lumin/MagicLeapTablet/Source/Private/MagicLeapTabletFunctionLibrary.cpp b/Engine/Plugins/Lumin/MagicLeapTablet/Source/Private/MagicLeapTabletFunctionLibrary.cpp
new file mode 100644
--- /dev/null
+++ b/Engine/Plugins/Lumin/MagicLeapTablet/Source/Private/MagicLeapTabletFunctionLibrary.cpp
@@ -0,0 +1,46 @@
+// Copyright 1998-2019 Epic Games, Inc. All Rights Reserved.
+
+#include "MagicLeapTabletFunctionLibrary.h"
+#include "MagicLeapTabletPlugin.h"
+
+void UMagicLeapTabletFunctionLibrary::SetTabletConnectedDelegate(const FMagicLeapTabletOnConnectedDelegate& InConnectedDelegate)
+{
+	FMagicLeapTabletOnConnectedDelegateMulti ConnectedDelegate;
+	ConnectedDelegate.Add(InConnectedDelegate);
+	GetMagicLeapTabletPlugin().SetConnectedDelegate(ConnectedDelegate);
+}
+
+void UMagicLeapTabletFunctionLibrary::SetTabletDisconnectedDelegate(const FMagicLeapTabletOnDisconnectedDelegate& InDisconnectedDelegate)
+{
+	FMagicLeapTabletOnDisconnectedDelegateMulti DisconnectedDelegate;
+	DisconnectedDelegate.Add(InDisconnectedDelegate);
+	GetMagicLeapTabletPlugin().SetDisconnectedDelegate(DisconnectedDelegate);
+}
+
+void UMagicLeapTabletFunctionLibrary::SetPenTouchDelegate(const FMagicLeapTabletOnPenTouchDelegate& InPenTouchDelegate)
+{
+	FMagicLeapTabletOnPenTouchDelegateMulti PenTouchDelegate;
+	PenTouchDelegate.Add(InPenTouchDelegate);
+	GetMagicLeapTabletPlugin().SetPenTouchDelegate(PenTouchDelegate);
+}
+
+void UMagicLeapTabletFunctionLibrary::SetRingTouchDelegate(const FMagicLeapTabletOnRingTouchDelegate& InRingTouchDelegate)
+{
+	FMagicLeapTabletOnRingTouchDelegateMulti RingTouchDelegate;
+	RingTouchDelegate.Add(InRingTouchDelegate);
+	GetMagicLeapTabletPlugin().SetRingTouchDelegate(RingTouchDelegate);
+}
+
+void UMagicLeapTabletFunctionLibrary::SetButtonDownDelegate(const FMagicLeapTabletOnButtonDownDelegate& InButtonDownDelegate)
+{
+	FMagicLeapTabletOnButtonDownDelegateMulti ButtonDownDelegate;
+	ButtonDownDelegate.Add(InButtonDownDelegate);
+	GetMagicLeapTabletPlugin().SetButtonDownDelegate(ButtonDownDelegate);
+}
+
+void UMagicLeapTabletFunctionLibrary::SetButtonUpDelegate(const FMagicLeapTabletOnButtonUpDelegate& InButtonUpDelegate)
+{
+	FMagicLeapTabletOnButtonUpDelegateMulti ButtonUpDelegate;
+	ButtonUpDelegate.Add(InButtonUpDelegate);
+	GetMagicLeapTabletPlugin().SetButtonUpDelegate(ButtonUpDelegate);
+}
diff --git a/Engine/Plugins/Lumin/MagicLeapTablet/Source/Private/MagicLeapTabletPlugin.cpp b/Engine/Plugins/Lumin/MagicLeapTablet/Source/Private/MagicLeapTabletPlugin.cpp
new file mode 100644
--- /dev/null
+++ b/Engine/Plugins/Lumin/MagicLeapTablet/Source/Private/MagicLeapTabletPlugin.cpp
@@ -0,0 +1,317 @@
+// Copyright 1998-2019 Epic Games, Inc. All Rights Reserved.
+
+#include "MagicLeapTabletPlugin.h"
+#include "IMagicLeapPlugin.h"
+#include "Engine/Engine.h"
+#include "MagicLeapHandle.h"
+#include "MagicLeapMath.h"
+
+DEFINE_LOG_CATEGORY(LogMagicLeapTablet);
+
+FMagicLeapTabletPlugin::FMagicLeapTabletPlugin()
+#if WITH_MLSDK
+: InputTracker(ML_INVALID_HANDLE)
+#endif // WITH_MLSDK
+{
+	IMagicLeapPlugin::Get().RegisterMagicLeapTrackerEntity(this);
+}
+
+FMagicLeapTabletPlugin::~FMagicLeapTabletPlugin()
+{
+	if (IMagicLeapPlugin::IsAvailable())
+	{
+		IMagicLeapPlugin::Get().UnregisterMagicLeapTrackerEntity(this);
+	}
+}
+
+void FMagicLeapTabletPlugin::CreateEntityTracker()
+{
+#if WITH_MLSDK
+	MLInputConfiguration InputConfig = { { MLInputControllerDof_6, MLInputControllerDof_6 } };
+	MLResult Result = MLInputCreate(&InputConfig, &InputTracker);
+	if (Result != MLResult_Ok)
+	{
+		UE_LOG(LogMagicLeapTablet, Error, TEXT("MLInputCreate failed with error %s."), UTF8_TO_TCHAR(MLGetResultString(Result)));
+		return;
+	}
+
+	MLInputTabletDeviceCallbacksInit(&Callbacks);
+	TickDelegate = FTickerDelegate::CreateRaw(this, &FMagicLeapTabletPlugin::Tick);
+	TickDelegateHandle = FTicker::GetCoreTicker().AddTicker(TickDelegate);
+#endif // WITH_MLSDK
+}
+
+void FMagicLeapTabletPlugin::DestroyEntityTracker()
+{
+#if WITH_MLSDK
+	if (MLHandleIsValid(InputTracker))
+	{
+		MLResult Result = MLInputSetTabletDeviceCallbacks(InputTracker, nullptr, nullptr);
+		UE_CLOG(Result != MLResult_Ok, LogMagicLeapTablet, Error, TEXT("MLInputSetTabletDeviceCallbacks failed with error %s!"), UTF8_TO_TCHAR(MLGetResultString(Result)));
+		Result = MLInputDestroy(InputTracker);
+		InputTracker = ML_INVALID_HANDLE;
+		UE_CLOG(Result != MLResult_Ok, LogMagicLeapTablet, Error, TEXT("MLInputDestroy failed with error %s!"), UTF8_TO_TCHAR(MLGetResultString(Result)));
+	}
+
+	FTicker::GetCoreTicker().RemoveTicker(TickDelegateHandle);
+#endif // WITH_MLSDK
+}
+
+bool FMagicLeapTabletPlugin::Tick(float DeltaTime)
+{
+	FPendingCallbackData PendingCallbackData;
+	while (PendingCallbacks.Dequeue(PendingCallbackData))
+	{
+		switch (PendingCallbackData.Type)
+		{
+		case ECallbackType::Connect: OnConnected.Broadcast(PendingCallbackData.DeviceId); break;
+		case ECallbackType::Disconnect: OnDisconnected.Broadcast(PendingCallbackData.DeviceId); break;
+		case ECallbackType::PenTouch: 
+		{
+#if WITH_MLSDK
+			MLInputTabletDeviceStatesList TabletDeviceStates;
+			MLInputTabletDeviceStatesListInit(&TabletDeviceStates);
+			MLResult Result = MLInputGetTabletDeviceStates(InputTracker, PendingCallbackData.DeviceId, &TabletDeviceStates);
+			if (Result != MLResult_Ok)
+			{
+				UE_LOG(LogMagicLeapTablet, Error, TEXT("MLInputGetTabletDeviceStates failed with error '%s'!"), UTF8_TO_TCHAR(MLGetResultString(Result)));
+			}
+			else if (TabletDeviceStates.count > 0 && TabletDeviceStates.data[TabletDeviceStates.count - 1].valid_fields_flag & MLInputTabletDeviceStateMask_HasToolType)
+			{
+				PendingCallbackData.DeviceState.ToolType = MLToUnrealDeviceToolType(TabletDeviceStates.data[TabletDeviceStates.count - 1].tool_type);
+			}
+			Result = MLInputReleaseTabletDeviceStates(InputTracker, &TabletDeviceStates);
+			UE_CLOG(Result != MLResult_Ok, LogMagicLeapTablet, Error, TEXT("MLInputReleaseTabletDeviceStates failed with error '%s'!"), UTF8_TO_TCHAR(MLGetResultString(Result)));
+#endif // WITH_MLSDK
+			OnPenTouch.Broadcast(PendingCallbackData.DeviceId, PendingCallbackData.DeviceState);
+		}
+		break;
+		case ECallbackType::RingTouch: OnRingTouch.Broadcast(PendingCallbackData.DeviceId, PendingCallbackData.RingTouchValue); break;
+		case ECallbackType::ButtonDown: OnButtonDown.Broadcast(PendingCallbackData.DeviceId, PendingCallbackData.Button); break;
+		case ECallbackType::ButtonUp: OnButtonUp.Broadcast(PendingCallbackData.DeviceId, PendingCallbackData.Button); break;
+		}
+	}
+
+	return true;
+}
+
+void FMagicLeapTabletPlugin::SetConnectedDelegate(const FMagicLeapTabletOnConnectedDelegateMulti& ConnectedDelegate)
+{
+	OnConnected = ConnectedDelegate;
+#if WITH_MLSDK
+	Callbacks.on_connect = OnConnectEvent;
+	MLResult Result = MLInputSetTabletDeviceCallbacks(InputTracker, &Callbacks, this);
+	if (Result != MLResult_Ok)
+	{
+		UE_LOG(LogMagicLeapTablet, Error, TEXT("MLInputSetTabletDeviceCallbacks failed with error %s!"), UTF8_TO_TCHAR(MLGetResultString(Result)));
+		return;
+	}
+
+	MLInputConnectedDevicesList ConnectedDevices;
+	MLInputConnectedDevicesListInit(&ConnectedDevices);
+
+	Result = MLInputGetConnectedDevices(InputTracker, &ConnectedDevices);
+	if (Result != MLResult_Ok)
+	{
+		UE_LOG(LogMagicLeapTablet, Error, TEXT("MLInputGetConnectedDevices failed with error %s."), UTF8_TO_TCHAR(MLGetResultString(Result)));
+		return;
+	}
+
+	for (uint32_t iConnectedTablet = 0; iConnectedTablet < ConnectedDevices.connected_tablet_device_count; ++iConnectedTablet)
+	{
+		ConnectedDelegate.Broadcast(ConnectedDevices.connected_tablet_device_ids[iConnectedTablet]);
+	}
+
+	Result = MLInputReleaseConnectedDevicesList(InputTracker, &ConnectedDevices);
+	UE_CLOG(Result == MLResult_Ok, LogMagicLeapTablet, Error, TEXT("MLInputReleaseConnectedDevicesList failed with error %s."), UTF8_TO_TCHAR(MLGetResultString(Result)));
+#endif // WITH_MLSDK
+}
+
+void FMagicLeapTabletPlugin::SetDisconnectedDelegate(const FMagicLeapTabletOnDisconnectedDelegateMulti& DisconnectedDelegate)
+{
+	OnDisconnected = DisconnectedDelegate;
+#if WITH_MLSDK
+	Callbacks.on_disconnect = OnDisconnectEvent;
+	MLResult Result = MLInputSetTabletDeviceCallbacks(InputTracker, &Callbacks, this);
+	UE_CLOG(Result != MLResult_Ok, LogMagicLeapTablet, Error, TEXT("MLInputSetTabletDeviceCallbacks failed with error %s!"), UTF8_TO_TCHAR(MLGetResultString(Result)));
+#endif // WITH_MLSDK
+}
+
+void FMagicLeapTabletPlugin::SetPenTouchDelegate(const FMagicLeapTabletOnPenTouchDelegateMulti& PenTouchDelegate)
+{
+	OnPenTouch = PenTouchDelegate;
+#if WITH_MLSDK
+	Callbacks.on_pen_touch_event = OnPenTouchEvent;
+	MLResult Result = MLInputSetTabletDeviceCallbacks(InputTracker, &Callbacks, this);
+	UE_CLOG(Result != MLResult_Ok, LogMagicLeapTablet, Error, TEXT("MLInputSetTabletDeviceCallbacks failed with error %s!"), UTF8_TO_TCHAR(MLGetResultString(Result)));
+#endif // WITH_MLSDK
+}
+
+void FMagicLeapTabletPlugin::SetRingTouchDelegate(const FMagicLeapTabletOnRingTouchDelegateMulti& RingTouchDelegate)
+{
+	OnRingTouch = RingTouchDelegate;
+#if WITH_MLSDK
+	Callbacks.on_touch_ring_event = OnRingTouchEvent;
+	MLResult Result = MLInputSetTabletDeviceCallbacks(InputTracker, &Callbacks, this);
+	UE_CLOG(Result != MLResult_Ok, LogMagicLeapTablet, Error, TEXT("MLInputSetTabletDeviceCallbacks failed with error %s!"), UTF8_TO_TCHAR(MLGetResultString(Result)));
+#endif // WITH_MLSDK
+}
+
+void FMagicLeapTabletPlugin::SetButtonDownDelegate(const FMagicLeapTabletOnButtonDownDelegateMulti& ButtonDownDelegate)
+{
+	OnButtonDown = ButtonDownDelegate;
+#if WITH_MLSDK
+	Callbacks.on_button_down = OnButtonDownEvent;
+	MLResult Result = MLInputSetTabletDeviceCallbacks(InputTracker, &Callbacks, this);
+	UE_CLOG(Result != MLResult_Ok, LogMagicLeapTablet, Error, TEXT("MLInputSetTabletDeviceCallbacks failed with error %s!"), UTF8_TO_TCHAR(MLGetResultString(Result)));
+#endif // WITH_MLSDK
+}
+
+void FMagicLeapTabletPlugin::SetButtonUpDelegate(const FMagicLeapTabletOnButtonUpDelegateMulti& ButtonUpDelegate)
+{
+	OnButtonUp = ButtonUpDelegate;
+#if WITH_MLSDK
+	Callbacks.on_button_up = OnButtonUpEvent;
+	MLResult Result = MLInputSetTabletDeviceCallbacks(InputTracker, &Callbacks, this);
+	UE_CLOG(Result != MLResult_Ok, LogMagicLeapTablet, Error, TEXT("MLInputSetTabletDeviceCallbacks failed with error %s!"), UTF8_TO_TCHAR(MLGetResultString(Result)));
+#endif // WITH_MLSDK
+}
+
+void FMagicLeapTabletPlugin::OnConnectEvent(uint8_t TabletDeviceId, void* Data)
+{
+	FPendingCallbackData PendingCallbackData;
+	PendingCallbackData.DeviceId = TabletDeviceId;
+	PendingCallbackData.Type = ECallbackType::Connect;
+	static_cast<FMagicLeapTabletPlugin*>(Data)->PendingCallbacks.Enqueue(PendingCallbackData);
+}
+
+void FMagicLeapTabletPlugin::OnDisconnectEvent(uint8_t TabletDeviceId, void* Data)
+{
+	FPendingCallbackData PendingCallbackData;
+	PendingCallbackData.DeviceId = TabletDeviceId;
+	PendingCallbackData.Type = ECallbackType::Disconnect;
+	static_cast<FMagicLeapTabletPlugin*>(Data)->PendingCallbacks.Enqueue(PendingCallbackData);
+}
+
+#if WITH_MLSDK
+void FMagicLeapTabletPlugin::OnPenTouchEvent(uint8_t TabletDeviceId, const MLInputTabletDeviceState* TabletDeviceState, void* Data)
+{
+	FPendingCallbackData PendingCallbackData;
+	PendingCallbackData.DeviceId = TabletDeviceId;
+	PendingCallbackData.Type = ECallbackType::PenTouch;
+	MLToUnrealTabletDeviceState(TabletDeviceState, PendingCallbackData.DeviceState);
+	static_cast<FMagicLeapTabletPlugin*>(Data)->PendingCallbacks.Enqueue(PendingCallbackData);
+}
+
+void FMagicLeapTabletPlugin::OnRingTouchEvent(uint8_t TabletDeviceId, int32_t RingTouchValue, uint64_t Timestamp, void *Data)
+{
+	FPendingCallbackData PendingCallbackData;
+	PendingCallbackData.DeviceId = TabletDeviceId;
+	PendingCallbackData.Type = ECallbackType::RingTouch;
+	PendingCallbackData.RingTouchValue = RingTouchValue;
+	static_cast<FMagicLeapTabletPlugin*>(Data)->PendingCallbacks.Enqueue(PendingCallbackData);
+}
+
+void FMagicLeapTabletPlugin::OnButtonDownEvent(uint8_t TabletDeviceId, MLInputTabletDeviceButton Button, uint64_t Timestamp, void* Data)
+{
+	FPendingCallbackData PendingCallbackData;
+	PendingCallbackData.DeviceId = TabletDeviceId;
+	PendingCallbackData.Type = ECallbackType::ButtonDown;
+	PendingCallbackData.Button = MLToUnrealInputTabletDeviceButton(Button);
+	static_cast<FMagicLeapTabletPlugin*>(Data)->PendingCallbacks.Enqueue(PendingCallbackData);
+}
+
+void FMagicLeapTabletPlugin::OnButtonUpEvent(uint8_t TabletDeviceId, MLInputTabletDeviceButton Button, uint64_t Timestamp, void* Data)
+{
+	FPendingCallbackData PendingCallbackData;
+	PendingCallbackData.DeviceId = TabletDeviceId;
+	PendingCallbackData.Type = ECallbackType::ButtonUp;
+	PendingCallbackData.Button = MLToUnrealInputTabletDeviceButton(Button);
+	static_cast<FMagicLeapTabletPlugin*>(Data)->PendingCallbacks.Enqueue(PendingCallbackData);
+}
+
+EMagicLeapInputTabletDeviceType FMagicLeapTabletPlugin::MLToUnrealInputTabletDeviceType(MLInputTabletDeviceType InDeviceType)
+{
+	EMagicLeapInputTabletDeviceType Type = EMagicLeapInputTabletDeviceType::Unknown;
+
+	switch (InDeviceType)
+	{
+	case MLInputTabletDeviceType_Unknown: Type = EMagicLeapInputTabletDeviceType::Unknown; break;
+	case MLInputTabletDeviceType_Wacom: Type = EMagicLeapInputTabletDeviceType::Wacom; break;
+	case MLInputTabletDeviceType_Ensure32Bits:
+	default: UE_LOG(LogMagicLeapTablet, Error, TEXT("Unexpected type encountered!"));
+	}
+
+	return Type;
+}
+
+EMagicLeapInputTabletDeviceToolType FMagicLeapTabletPlugin::MLToUnrealDeviceToolType(MLInputTabletDeviceToolType InputDeviceToolType)
+{
+	EMagicLeapInputTabletDeviceToolType ToolType = EMagicLeapInputTabletDeviceToolType::Unknown;
+
+	switch (InputDeviceToolType)
+	{
+	case MLInputTabletDeviceToolType_Unknown: ToolType = EMagicLeapInputTabletDeviceToolType::Unknown; break;
+	case MLInputTabletDeviceToolType_Pen: ToolType = EMagicLeapInputTabletDeviceToolType::Pen; break;
+	case MLInputTabletDeviceToolType_Eraser: ToolType = EMagicLeapInputTabletDeviceToolType::Eraser; break;
+	case MLInputTabletDeviceTootlType_Ensure32Bits:
+	default: UE_LOG(LogMagicLeapTablet, Error, TEXT("Unexpected type encountered!"));
+	}
+
+	return ToolType;
+}
+
+EMagicLeapInputTabletDeviceButton FMagicLeapTabletPlugin::MLToUnrealInputTabletDeviceButton(MLInputTabletDeviceButton InMLButton)
+{
+	EMagicLeapInputTabletDeviceButton UnrealButton = EMagicLeapInputTabletDeviceButton::Unknown;
+
+	switch (InMLButton)
+	{
+	case MLInputTabletDeviceButton_Unknown: UnrealButton = EMagicLeapInputTabletDeviceButton::Unknown; break;
+	case MLInputTabletDeviceButton_Button1: UnrealButton = EMagicLeapInputTabletDeviceButton::Button1; break;
+	case MLInputTabletDeviceButton_Button2: UnrealButton = EMagicLeapInputTabletDeviceButton::Button2; break;
+	case MLInputTabletDeviceButton_Button3: UnrealButton = EMagicLeapInputTabletDeviceButton::Button3; break;
+	case MLInputTabletDeviceButton_Button4: UnrealButton = EMagicLeapInputTabletDeviceButton::Button4; break;
+	case MLInputTabletDeviceButton_Button5: UnrealButton = EMagicLeapInputTabletDeviceButton::Button5; break;
+	case MLInputTabletDeviceButton_Button6: UnrealButton = EMagicLeapInputTabletDeviceButton::Button6; break;
+	case MLInputTabletDeviceButton_Button7: UnrealButton = EMagicLeapInputTabletDeviceButton::Button7; break;
+	case MLInputTabletDeviceButton_Button8: UnrealButton = EMagicLeapInputTabletDeviceButton::Button8; break;
+	case MLInputTabletDeviceButton_Button9: UnrealButton = EMagicLeapInputTabletDeviceButton::Button9; break;
+	case MLInputTabletDeviceButton_Button10: UnrealButton = EMagicLeapInputTabletDeviceButton::Button10; break;
+	case MLInputTabletDeviceButton_Button11: UnrealButton = EMagicLeapInputTabletDeviceButton::Button11; break;
+	case MLInputTabletDeviceButton_Button12: UnrealButton = EMagicLeapInputTabletDeviceButton::Button12; break;
+	case MLInputTabletDeviceButton_Count:
+	case MLInputTabletDeviceButton_Ensure32Bits:
+	default: UE_LOG(LogMagicLeapTablet, Error, TEXT("Unexpected type encountered!"));
+	}
+
+	return UnrealButton;
+}
+
+void FMagicLeapTabletPlugin::MLToUnrealTabletDeviceState(const MLInputTabletDeviceState* InMLInputTabletDeviceState, FMagicLeapInputTabletDeviceState& OutUnrealTabletDeviceState)
+{
+	OutUnrealTabletDeviceState.Version = InMLInputTabletDeviceState->version;
+	OutUnrealTabletDeviceState.Type = MLToUnrealInputTabletDeviceType(InMLInputTabletDeviceState->type);
+	OutUnrealTabletDeviceState.ToolType = MLToUnrealDeviceToolType(InMLInputTabletDeviceState->tool_type);
+	OutUnrealTabletDeviceState.PenTouchPosAndForce = 
+	{
+		InMLInputTabletDeviceState->pen_touch_pos_and_force.x,
+		InMLInputTabletDeviceState->pen_touch_pos_and_force.y,
+		InMLInputTabletDeviceState->pen_touch_pos_and_force.z
+	};
+	OutUnrealTabletDeviceState.AdditionalPenTouchData =
+	{
+		InMLInputTabletDeviceState->additional_pen_touch_data[0],
+		InMLInputTabletDeviceState->additional_pen_touch_data[1],
+		InMLInputTabletDeviceState->additional_pen_touch_data[2]
+	};
+	OutUnrealTabletDeviceState.bIsPenTouchActive = InMLInputTabletDeviceState->is_pen_touch_active;
+	OutUnrealTabletDeviceState.bIsConnected = InMLInputTabletDeviceState->is_connected;
+	OutUnrealTabletDeviceState.PenDistance = InMLInputTabletDeviceState->pen_distance;
+	// Converting from nanoseconds to microseconds by dividing by 1000
+	OutUnrealTabletDeviceState.Timestamp = FTimespan::FromMicroseconds(InMLInputTabletDeviceState->timestamp / 1000.0f);
+	OutUnrealTabletDeviceState.ValidFieldsFlag = InMLInputTabletDeviceState->valid_fields_flag;
+}
+#endif // WITH_MLSDK
+
+IMPLEMENT_MODULE(FMagicLeapTabletPlugin, MagicLeapTablet);
diff --git a/Engine/Plugins/Lumin/MagicLeapTablet/Source/Private/MagicLeapTabletPlugin.h b/Engine/Plugins/Lumin/MagicLeapTablet/Source/Private/MagicLeapTabletPlugin.h
new file mode 100644
--- /dev/null
+++ b/Engine/Plugins/Lumin/MagicLeapTablet/Source/Private/MagicLeapTabletPlugin.h
@@ -0,0 +1,81 @@
+// Copyright 1998-2019 Epic Games, Inc. All Rights Reserved.
+
+#pragma once
+
+#include "CoreMinimal.h"
+#include "Containers/Ticker.h"
+#include "IMagicLeapTabletPlugin.h"
+#include "IMagicLeapTrackerEntity.h"
+#include "Lumin/CAPIShims/LuminAPITablet.h"
+#include "Containers/Queue.h"
+
+DECLARE_LOG_CATEGORY_EXTERN(LogMagicLeapTablet, Verbose, All);
+
+class FMagicLeapTabletPlugin : public IMagicLeapTabletPlugin, public IMagicLeapTrackerEntity
+{
+public:
+	FMagicLeapTabletPlugin();
+	virtual ~FMagicLeapTabletPlugin();
+
+	void CreateEntityTracker() override;
+	void DestroyEntityTracker() override;
+	bool Tick(float DeltaTime);
+
+	void SetConnectedDelegate(const FMagicLeapTabletOnConnectedDelegateMulti& ConnectedDelegate);
+	void SetDisconnectedDelegate(const FMagicLeapTabletOnDisconnectedDelegateMulti& DisconnectedDelegate);
+	void SetPenTouchDelegate(const FMagicLeapTabletOnPenTouchDelegateMulti& PenTouchDelegate);
+	void SetRingTouchDelegate(const FMagicLeapTabletOnRingTouchDelegateMulti& RingTouchDelegate);
+	void SetButtonDownDelegate(const FMagicLeapTabletOnButtonDownDelegateMulti& ButtonDownDelegate);
+	void SetButtonUpDelegate(const FMagicLeapTabletOnButtonUpDelegateMulti& ButtonDownDelegate);
+
+private:
+	FTickerDelegate TickDelegate;
+	FDelegateHandle TickDelegateHandle;
+	FMagicLeapTabletOnConnectedDelegateMulti OnConnected;
+	FMagicLeapTabletOnDisconnectedDelegateMulti OnDisconnected;
+	FMagicLeapTabletOnPenTouchDelegateMulti OnPenTouch;
+	FMagicLeapTabletOnRingTouchDelegateMulti OnRingTouch;
+	FMagicLeapTabletOnButtonDownDelegateMulti OnButtonDown;
+	FMagicLeapTabletOnButtonUpDelegateMulti OnButtonUp;
+	enum ECallbackType
+	{
+		Connect,
+		Disconnect,
+		PenTouch,
+		RingTouch,
+		ButtonDown,
+		ButtonUp
+	};
+	struct FPendingCallbackData
+	{
+		uint8 DeviceId;
+		ECallbackType Type;
+		FMagicLeapInputTabletDeviceState DeviceState;
+		int32_t RingTouchValue;
+		EMagicLeapInputTabletDeviceButton Button;
+	};
+	TQueue<FPendingCallbackData, EQueueMode::Spsc> PendingCallbacks;
+#if WITH_MLSDK
+	MLHandle InputTracker;
+	MLInputTabletDeviceCallbacks Callbacks;
+#endif // WITH_MLSDK
+
+	static void OnConnectEvent(uint8_t TabletDeviceId, void* Data);
+	static void OnDisconnectEvent(uint8_t TabletDeviceId, void* Data);
+#if WITH_MLSDK
+	static void OnPenTouchEvent(uint8_t TabletDeviceId, const MLInputTabletDeviceState* TabletDeviceState, void* Data);
+	static void OnRingTouchEvent(uint8_t TabletDeviceId, int32_t RingTouchValue, uint64_t Timestamp, void *Data);
+	static void OnButtonDownEvent(uint8_t TabletDeviceId, MLInputTabletDeviceButton Button, uint64_t Timestamp, void* Data);
+	static void OnButtonUpEvent(uint8_t TabletDeviceId, MLInputTabletDeviceButton Button, uint64_t Timestamp, void* Data);
+
+	static EMagicLeapInputTabletDeviceType MLToUnrealInputTabletDeviceType(MLInputTabletDeviceType InDeviceType);
+	static EMagicLeapInputTabletDeviceToolType MLToUnrealDeviceToolType(MLInputTabletDeviceToolType InputDeviceToolType);
+	static EMagicLeapInputTabletDeviceButton MLToUnrealInputTabletDeviceButton(MLInputTabletDeviceButton InMLButton);
+	static void MLToUnrealTabletDeviceState(const MLInputTabletDeviceState* InMLInputTabletDeviceState, FMagicLeapInputTabletDeviceState& OutUnrealTabletDeviceState);
+#endif // WITH_MLSDK
+};
+
+inline FMagicLeapTabletPlugin& GetMagicLeapTabletPlugin()
+{
+	return FModuleManager::Get().GetModuleChecked<FMagicLeapTabletPlugin>("MagicLeapTablet");
+}
diff --git a/Engine/Plugins/Lumin/MagicLeapTablet/Source/Public/IMagicLeapTabletPlugin.h b/Engine/Plugins/Lumin/MagicLeapTablet/Source/Public/IMagicLeapTabletPlugin.h
new file mode 100644
--- /dev/null
+++ b/Engine/Plugins/Lumin/MagicLeapTablet/Source/Public/IMagicLeapTabletPlugin.h
@@ -0,0 +1,36 @@
+// Copyright 1998-2019 Epic Games, Inc. All Rights Reserved.
+
+#pragma once
+
+#include "Modules/ModuleManager.h"
+#include "MagicLeapTabletTypes.h"
+
+	/**
+	 * The public interface to this module.  In most cases, this interface is only public to sibling modules
+	 * within this plugin.
+	 */
+class IMagicLeapTabletPlugin : public IModuleInterface
+{
+public:
+
+	/**
+	 * Singleton-like access to this module's interface.  This is just for convenience!
+	 * Beware of calling this during the shutdown phase, though.  Your module might have been unloaded already.
+	 *
+	 * @return Returns singleton instance, loading the module on demand if needed
+	 */
+	static inline IMagicLeapTabletPlugin& Get()
+	{
+		return FModuleManager::LoadModuleChecked<IMagicLeapTabletPlugin>("MagicLeapTablet");
+	}
+
+	/**
+	 * Checks to see if this module is loaded and ready.  It is only valid to call Get() if IsAvailable() returns true.
+	 *
+	 * @return True if the module is loaded and ready to use
+	 */
+	static inline bool IsAvailable()
+	{
+		return FModuleManager::Get().IsModuleLoaded("MagicLeapTablet");
+	}
+};
diff --git a/Engine/Plugins/Lumin/MagicLeapTablet/Source/Public/MagicLeapTabletComponent.h b/Engine/Plugins/Lumin/MagicLeapTablet/Source/Public/MagicLeapTabletComponent.h
new file mode 100644
--- /dev/null
+++ b/Engine/Plugins/Lumin/MagicLeapTablet/Source/Public/MagicLeapTabletComponent.h
@@ -0,0 +1,43 @@
+// Copyright 1998-2019 Epic Games, Inc. All Rights Reserved.
+
+#pragma once
+
+#include "Engine/Engine.h"
+#include "Components/ActorComponent.h"
+#include "MagicLeapTabletTypes.h"
+#include "MagicLeapTabletComponent.generated.h"
+
+/**
+	Component that provides access to the Tablet API functionality.
+*/
+UCLASS(ClassGroup = MagicLeap, BlueprintType, Blueprintable, EditInlineNew, meta = (BlueprintSpawnableComponent))
+class MAGICLEAPTABLET_API UMagicLeapTabletComponent : public USceneComponent
+{
+	GENERATED_BODY()
+
+public:
+	UMagicLeapTabletComponent();
+
+	/** Passes the component delegates to the underlying tablet plugin. */
+	virtual void BeginPlay() override;
+
+private:
+	// Delegate instances
+	UPROPERTY(BlueprintAssignable, Category = "Tablet | MagicLeap", meta = (AllowPrivateAccess = true))
+	FMagicLeapTabletOnConnectedDelegateMulti OnConnected;
+
+	UPROPERTY(BlueprintAssignable, Category = "Tablet | MagicLeap", meta = (AllowPrivateAccess = true))
+	FMagicLeapTabletOnDisconnectedDelegateMulti OnDisconnected;
+
+	UPROPERTY(BlueprintAssignable, Category = "Tablet | MagicLeap", meta = (AllowPrivateAccess = true))
+	FMagicLeapTabletOnPenTouchDelegateMulti OnPenTouch;
+
+	UPROPERTY(BlueprintAssignable, Category = "Tablet | MagicLeap", meta = (AllowPrivateAccess = true))
+	FMagicLeapTabletOnRingTouchDelegateMulti OnRingTouch;
+
+	UPROPERTY(BlueprintAssignable, Category = "Tablet | MagicLeap", meta = (AllowPrivateAccess = true))
+	FMagicLeapTabletOnButtonDownDelegateMulti OnButtonDown;
+
+	UPROPERTY(BlueprintAssignable, Category = "Tablet | MagicLeap", meta = (AllowPrivateAccess = true))
+	FMagicLeapTabletOnButtonUpDelegateMulti OnButtonUp;
+};
diff --git a/Engine/Plugins/Lumin/MagicLeapTablet/Source/Public/MagicLeapTabletFunctionLibrary.h b/Engine/Plugins/Lumin/MagicLeapTablet/Source/Public/MagicLeapTabletFunctionLibrary.h
new file mode 100644
--- /dev/null
+++ b/Engine/Plugins/Lumin/MagicLeapTablet/Source/Public/MagicLeapTabletFunctionLibrary.h
@@ -0,0 +1,32 @@
+// Copyright 1998-2019 Epic Games, Inc. All Rights Reserved.
+#pragma once
+
+#include "Engine/Engine.h"
+#include "Kismet/BlueprintFunctionLibrary.h"
+#include "MagicLeapTabletTypes.h"
+#include "MagicLeapTabletFunctionLibrary.generated.h"
+
+UCLASS(ClassGroup = MagicLeap)
+class MAGICLEAPTABLET_API UMagicLeapTabletFunctionLibrary : public UBlueprintFunctionLibrary
+{
+	GENERATED_BODY()
+
+public:
+	UFUNCTION(BlueprintCallable, Category = "Tablet Function Library | MagicLeap")
+	static void SetTabletConnectedDelegate(const FMagicLeapTabletOnConnectedDelegate& ConnectedDelegate);
+
+	UFUNCTION(BlueprintCallable, Category = "Tablet Function Library | MagicLeap")
+	static void SetTabletDisconnectedDelegate(const FMagicLeapTabletOnDisconnectedDelegate& DisconnectedDelegate);
+
+	UFUNCTION(BlueprintCallable, Category = "Tablet Function Library | MagicLeap")
+	static void SetPenTouchDelegate(const FMagicLeapTabletOnPenTouchDelegate& PenTouchDelegate);
+
+	UFUNCTION(BlueprintCallable, Category = "Tablet Function Library | MagicLeap")
+	static void SetRingTouchDelegate(const FMagicLeapTabletOnRingTouchDelegate& RingTouchDelegate);
+
+	UFUNCTION(BlueprintCallable, Category = "Tablet Function Library | MagicLeap")
+	static void SetButtonDownDelegate(const FMagicLeapTabletOnButtonDownDelegate& ButtonDownDelegate);
+
+	UFUNCTION(BlueprintCallable, Category = "Tablet Function Library | MagicLeap")
+	static void SetButtonUpDelegate(const FMagicLeapTabletOnButtonUpDelegate& ButtonDownDelegate);
+};
diff --git a/Engine/Plugins/Lumin/MagicLeapTablet/Source/Public/MagicLeapTabletTypes.h b/Engine/Plugins/Lumin/MagicLeapTablet/Source/Public/MagicLeapTabletTypes.h
new file mode 100644
--- /dev/null
+++ b/Engine/Plugins/Lumin/MagicLeapTablet/Source/Public/MagicLeapTabletTypes.h
@@ -0,0 +1,124 @@
+// Copyright 1998-2019 Epic Games, Inc. All Rights Reserved.
+#pragma once
+
+#include "Engine/Engine.h"
+#include "MagicLeapTabletTypes.generated.h"
+
+UENUM(BlueprintType)
+enum class EMagicLeapInputTabletDeviceType : uint8
+{
+	Unknown,
+	Wacom,
+};
+
+UENUM(BlueprintType)
+enum class EMagicLeapInputTabletDeviceToolType : uint8
+{
+	Unknown,
+	Pen,
+	Eraser,
+};
+
+UENUM(BlueprintType)
+enum class EMagicLeapInputTabletDeviceButton : uint8
+{
+	Unknown,
+	Button1,
+	Button2,
+	Button3,
+	Button4,
+	Button5,
+	Button6,
+	Button7,
+	Button8,
+	Button9,
+	Button10,
+	Button11,
+	Button12,
+};
+
+/** Mask for determinig the validity of the additional pen data. */
+enum class EMagicLeapInputTabletDeviceStateMask : uint8
+{
+	HasType = 1 << 0,
+	HasToolType = 1 << 1,
+	HasPenTouchPosAndForce = 1 << 2,
+	HasAdditionalPenTouchData = 1 << 3,
+	HasPenTouchActive = 1 << 4,
+	HasConnectionState = 1 << 5,
+	HasPenDistance = 1 << 6,
+	HasTimestamp = 1 << 7,
+};
+
+USTRUCT(BlueprintType)
+struct FMagicLeapInputTabletDeviceState
+{
+	GENERATED_BODY();
+	/** Version of this structure. */
+	uint32 Version;
+
+	/** Type of this tablet device. */
+	UPROPERTY(EditAnywhere, BlueprintReadWrite, Category = "Tablet | MagicLeap")
+	EMagicLeapInputTabletDeviceType Type;
+
+	/** Type of tool used with the tablet. */
+	UPROPERTY(EditAnywhere, BlueprintReadWrite, Category = "Tablet | MagicLeap")
+	EMagicLeapInputTabletDeviceToolType ToolType;
+
+	/**
+	  The current touch position (x,y) and force (z).
+	  Position is in the [-1.0,1.0] range and force is in the [0.0,1.0] range.
+	 */
+	UPROPERTY(EditAnywhere, BlueprintReadWrite, Category = "Tablet | MagicLeap")
+	FVector PenTouchPosAndForce;
+
+	/**
+	  The Additional coordinate values (x, y, z).
+	  It could contain data specific to the device type.
+	  For example, it could hold tilt values while using pen.
+	 */
+	UPROPERTY(EditAnywhere, BlueprintReadWrite, Category = "Tablet | MagicLeap")
+	TArray<int32> AdditionalPenTouchData;
+
+	/** Is touch active. */
+	UPROPERTY(EditAnywhere, BlueprintReadWrite, Category = "Tablet | MagicLeap")
+	bool bIsPenTouchActive;
+
+	/** If this tablet is connected. */
+	UPROPERTY(EditAnywhere, BlueprintReadWrite, Category = "Tablet | MagicLeap")
+	bool bIsConnected;
+
+	/** Distance between pen and tablet. */
+	UPROPERTY(EditAnywhere, BlueprintReadWrite, Category = "Tablet | MagicLeap")
+	float PenDistance;
+
+	/** Time since device bootup. */
+	UPROPERTY(EditAnywhere, BlueprintReadWrite, Category = "Tablet | MagicLeap")
+	FTimespan Timestamp;
+
+	/**
+	  Flags to denote which of the above fields are valid.
+	  #EMagicLeapInputTabletDeviceStateMask defines the bitmap.
+	 */
+	UPROPERTY(EditAnywhere, BlueprintReadWrite, Category = "Tablet | MagicLeap")
+	int32 ValidFieldsFlag;
+};
+
+/** Delegate used to signal when a tablet device has connected to the application. */
+DECLARE_DYNAMIC_DELEGATE_OneParam(FMagicLeapTabletOnConnectedDelegate, uint8, DeviceId);
+DECLARE_DYNAMIC_MULTICAST_DELEGATE_OneParam(FMagicLeapTabletOnConnectedDelegateMulti, uint8, DeviceId);
+/** Delegate used to signal when a tablet device has disconnected from the application. */
+DECLARE_DYNAMIC_DELEGATE_OneParam(FMagicLeapTabletOnDisconnectedDelegate, uint8, DeviceId);
+DECLARE_DYNAMIC_MULTICAST_DELEGATE_OneParam(FMagicLeapTabletOnDisconnectedDelegateMulti, uint8, DeviceId);
+/** Delegate used to signal the pen touch state of a connected tablet. */
+DECLARE_DYNAMIC_DELEGATE_TwoParams(FMagicLeapTabletOnPenTouchDelegate, uint8, DeviceId, const FMagicLeapInputTabletDeviceState&, DeviceState);
+DECLARE_DYNAMIC_MULTICAST_DELEGATE_TwoParams(FMagicLeapTabletOnPenTouchDelegateMulti, uint8, DeviceId, const FMagicLeapInputTabletDeviceState&, DeviceState);
+/** Delegate used to signal the ring touch state of a connected tablet. */
+DECLARE_DYNAMIC_DELEGATE_TwoParams(FMagicLeapTabletOnRingTouchDelegate, uint8, DeviceId, const int32&, TouchValue);
+DECLARE_DYNAMIC_MULTICAST_DELEGATE_TwoParams(FMagicLeapTabletOnRingTouchDelegateMulti, uint8, DeviceId, const int32&, TouchValue);
+/** Delegate used to signal a button press from a connected tablet. */
+DECLARE_DYNAMIC_DELEGATE_TwoParams(FMagicLeapTabletOnButtonDownDelegate, uint8, DeviceId, const EMagicLeapInputTabletDeviceButton&, Button);
+DECLARE_DYNAMIC_MULTICAST_DELEGATE_TwoParams(FMagicLeapTabletOnButtonDownDelegateMulti, uint8, DeviceId, const EMagicLeapInputTabletDeviceButton&, Button);
+/** Delegate used to signal a button release from a connected tablet. */
+DECLARE_DYNAMIC_DELEGATE_TwoParams(FMagicLeapTabletOnButtonUpDelegate, uint8, DeviceId, const EMagicLeapInputTabletDeviceButton&, Button);
+DECLARE_DYNAMIC_MULTICAST_DELEGATE_TwoParams(FMagicLeapTabletOnButtonUpDelegateMulti, uint8, DeviceId, const EMagicLeapInputTabletDeviceButton&, Button);
diff --git a/Engine/Source/Developer/Lumin/LuminPlatformEditor/Private/LuminTargetSettingsDetails.cpp b/Engine/Source/Developer/Lumin/LuminPlatformEditor/Private/LuminTargetSettingsDetails.cpp
--- a/Engine/Source/Developer/Lumin/LuminPlatformEditor/Private/LuminTargetSettingsDetails.cpp
+++ b/Engine/Source/Developer/Lumin/LuminPlatformEditor/Private/LuminTargetSettingsDetails.cpp
@@ -25,6 +25,8 @@ DEFINE_LOG_CATEGORY_STATIC(LogLuminTargetSettingsDetail, Log, All);
 
 #define LOCTEXT_NAMESPACE "LuminTargetSettingsDetails"
 
+static constexpr float ExpireDuration = 5.0f;
+
 TSharedRef<IDetailCustomization> FLuminTargetSettingsDetails::MakeInstance()
 {
 	return MakeShareable(new FLuminTargetSettingsDetails);
@@ -60,14 +62,10 @@ FString FLuminTargetSettingsDetails::IconModelPathGetter()
 	FString Result;
 	if (IconModelPathProp->GetValue(Result) == FPropertyAccess::Success && !Result.IsEmpty())
 	{
-		if (SetupForPlatformAttribute.Get())
-		{
-			// If we did the setup, the values are game/project dir relative. So we root them at that
-			// for display. Otherwise we show the default values which will show engine exec relative
-			// paths to the platform build resource locations.
-			Result = FPaths::ProjectDir() / Result;
-		}
-		else
+		// If we did the setup, the values are game/project dir relative. So we display
+		// them as is. Otherwise we show the default values which will show engine exec relative
+		// paths to the platform build resource locations.
+		if (!SetupForPlatformAttribute.Get())
 		{
 			// Otherwise they are in the engine tree. But they could be outdate paths.
 			// In which case we use the hard-wired defaults.
@@ -89,14 +87,10 @@ FString FLuminTargetSettingsDetails::IconPortalPathGetter()
 	FString Result;
 	if (IconPortalPathProp->GetValue(Result) == FPropertyAccess::Success && !Result.IsEmpty())
 	{
-		if (SetupForPlatformAttribute.Get())
-		{
-			// If we did the setup, the values are game/project dir relative. So we root them at that
-			// for display. Otherwise we show the default values which will show engine exec relative
-			// paths to the platform build resource locations.
-			Result = FPaths::ProjectDir() / Result;
-		}
-		else
+		// If we did the setup, the values are game/project dir relative. So we display
+		// them as is. Otherwise we show the default values which will show engine exec relative
+		// paths to the platform build resource locations.
+		if (!SetupForPlatformAttribute.Get())
 		{
 			// Otherwise they are in the engine tree. But they could be outdate paths.
 			// In which case we use the hard-wired defaults.
@@ -116,10 +110,7 @@ FString FLuminTargetSettingsDetails::IconPortalPathGetter()
 FString FLuminTargetSettingsDetails::CertificateGetter()
 {
 	FString Result;
-	if (CertificateProp->GetValue(Result) == FPropertyAccess::Success && !Result.IsEmpty())
-	{
-		Result = FPaths::ProjectDir() / Result;
-	}
+	CertificateProp->GetValue(Result);
 	return Result;
 }
 
@@ -212,12 +203,18 @@ void FLuminTargetSettingsDetails::BuildAppTileSection(IDetailLayoutBuilder& Deta
 		LOCTEXT("IconModelLabel", "Icon Model"),
 		LOCTEXT("PickIconModelButton_Tooltip", "Select the icon model to use for the application. The files will be copied to the project build folder."),
 		FOnChoosePath::CreateRaw(this, &FLuminTargetSettingsDetails::OnPickDirectory),
-		FOnPickPath::CreateRaw(this, &FLuminTargetSettingsDetails::OnPickIconModelPath));
+		FOnPickPath::CreateRaw(this, &FLuminTargetSettingsDetails::OnPickIconModelPath),
+		false /*bClearButton*/,
+		LOCTEXT("NoClearButton_Tooltip", ""),
+		true /*bDisableUntilConfigured*/);
 	BuildPathPicker(DetailBuilder, AppTitleCategory, IconPortalPathAttribute,
 		LOCTEXT("IconPortalLabel", "Icon Portal"),
 		LOCTEXT("PickIconPortalButton_Tooltip", "Select the icon portal to use for the application. The files will be copied to the project build folder."),
 		FOnChoosePath::CreateRaw(this, &FLuminTargetSettingsDetails::OnPickDirectory),
-		FOnPickPath::CreateRaw(this, &FLuminTargetSettingsDetails::OnPickIconPortalPath));
+		FOnPickPath::CreateRaw(this, &FLuminTargetSettingsDetails::OnPickIconPortalPath),
+		false /*bClearButton*/,
+		LOCTEXT("NoClearButton_Tooltip", ""),
+		true /*bDisableUntilConfigured*/);
 
 	////////// UI for signing cert..
 
@@ -233,7 +230,11 @@ void FLuminTargetSettingsDetails::BuildAppTileSection(IDetailLayoutBuilder& Deta
 				LOCTEXT("PickCertificateFileDialogTitle", "Choose a certificate").ToString(),
 				FString::Printf(TEXT("%s (*.cert)|*.cert"), *FilterText));
 		}),
-		FOnPickPath::CreateRaw(this, &FLuminTargetSettingsDetails::OnPickCertificate));
+		FOnPickPath::CreateRaw(this, &FLuminTargetSettingsDetails::OnPickCertificate),
+		true /*bClearButton*/,
+		LOCTEXT("CertClearButton_Tooltip", "Clear the selected certificate."),
+		false /*bDisableUntilConfigured*/);
+
 }
 
 FReply FLuminTargetSettingsDetails::OnPickIconModelPath(const FString& DirPath)
@@ -242,7 +243,13 @@ FReply FLuminTargetSettingsDetails::OnPickIconModelPath(const FString& DirPath)
 	if (ProjectModelPath != DirPath)
 	{
 		// Copy the contents of the selected path to the project build path.
-		CopyDir(DirPath, ProjectModelPath);
+		const bool bModelFilesUpdated = CopyDir(DirPath, ProjectModelPath);
+		if (bModelFilesUpdated)
+		{
+			FNotificationInfo SuccessInfo(LOCTEXT("ModelFilesUpdatedLabel", "Icon model files updated."));
+			SuccessInfo.ExpireDuration = ExpireDuration;
+			FSlateNotificationManager::Get().AddNotification(SuccessInfo);
+		}
 	}
 	return FReply::Handled();
 }
@@ -253,7 +260,13 @@ FReply FLuminTargetSettingsDetails::OnPickIconPortalPath(const FString& DirPath)
 	if (ProjectPortalPath != DirPath)
 	{
 		// Copy the contents of the selected path to the project build path.
-		CopyDir(DirPath, ProjectPortalPath);
+		const bool bPortalFilesUpdated = CopyDir(DirPath, ProjectPortalPath);
+		if (bPortalFilesUpdated)
+		{
+			FNotificationInfo SuccessInfo(LOCTEXT("PortalFilesUpdatedLabel", "Icon portal files updated."));
+			SuccessInfo.ExpireDuration = ExpireDuration;
+			FSlateNotificationManager::Get().AddNotification(SuccessInfo);
+		}
 	}
 	return FReply::Handled();
 }
@@ -272,6 +285,26 @@ bool FLuminTargetSettingsDetails::CopyDir(FString SourceDir, FString TargetDir)
 	int FilesCopiedCount = 0;
 	IPlatformFile::GetPlatformPhysical().FindFilesRecursively(FilesToCopy, *SourceDir, nullptr);
 	FText Description = FText::FromString(FPaths::GetBaseFilename(TargetDir));
+
+	// Delete files already existing in TargetDir but not in SourceDir i.e. the ones that won't be replaced.
+	TArray<FString> ExistingFiles;
+	IPlatformFile::GetPlatformPhysical().FindFilesRecursively(ExistingFiles, *TargetDir, nullptr);
+	for (const FString& ExistingFile : ExistingFiles)
+	{
+		const FString ExistingFileInSourceDir = ExistingFile.Replace(*TargetDir, *SourceDir);
+		if (!IPlatformFile::GetPlatformPhysical().FileExists(*ExistingFileInSourceDir))
+		{
+			if (SourceControlHelpers::IsEnabled())
+			{
+				SourceControlHelpers::MarkFileForDelete(ExistingFile);
+			}
+			else
+			{
+				IPlatformFile::GetPlatformPhysical().DeleteFile(*ExistingFile);
+			}
+		}
+	}
+
 	for (FString& FileToCopy : FilesToCopy)
 	{
 		FString NewFile = FileToCopy.Replace(*SourceDir, *TargetDir);
@@ -295,13 +328,15 @@ bool FLuminTargetSettingsDetails::CopyDir(FString SourceDir, FString TargetDir)
 			}
 		}
 	}
-	return true;
+
+	return FilesCopiedCount > 0;
 }
 
 void FLuminTargetSettingsDetails::BuildPathPicker(
 	IDetailLayoutBuilder& DetailBuilder, IDetailCategoryBuilder& Category,
 	TAttribute<FString> Path, const FText& Label, const FText& Tooltip,
-	const FOnChoosePath& OnChoose, const FOnPickPath& OnPick)
+	const FOnChoosePath& OnChoose, const FOnPickPath& OnPick, bool bClearButton, const FText& ClearTooltip,
+	bool bDisableUntilConfigured)
 {
 	TSharedPtr<SButton> PickButton = nullptr;
 	TSharedPtr<SWidget> PickWidget = nullptr;
@@ -322,8 +357,54 @@ void FLuminTargetSettingsDetails::BuildPathPicker(
 		return OnChoose.Execute(Path.Get(), OnPick, PickButton);
 	}));
 
-	Category.AddCustomRow(Label, false)
-		.IsEnabled(SetupForPlatformAttribute)
+	TSharedPtr<SButton> ClearButton = nullptr;
+	TSharedPtr<SWidget> ClearWidget = nullptr;
+	ClearWidget = SAssignNew(ClearButton, SButton)
+		.ButtonStyle(FEditorStyle::Get(), "HoverHintOnly")
+		.ToolTipText(ClearTooltip)
+		.ContentPadding(2.0f)
+		.ForegroundColor(FSlateColor::UseForeground())
+		.IsFocusable(false)
+		.IsEnabled(bClearButton)
+		[
+			SNew(SImage)
+			.Image(FEditorStyle::GetBrush("PropertyWindow.Button_Clear"))
+			.ColorAndOpacity(FSlateColor::UseForeground())
+		]
+	;
+	ClearButton->SetOnClicked(FOnClicked::CreateLambda([this]()->FReply
+	{
+		CertificateProp->SetValue(TEXT(""));
+		return FReply::Handled();
+	}));
+
+	TSharedPtr<SHorizontalBox> PathPickerBox = SNew(SHorizontalBox)
+		+ SHorizontalBox::Slot()
+		.AutoWidth()
+		[
+			SNew(STextBlock)
+			.Text(TAttribute<FText>::Create(
+				TAttribute<FText>::FGetter::CreateLambda([Path]()->FText { return FText::FromString(Path.Get()); })))
+			.Font(DetailBuilder.GetDetailFont())
+			.Margin(2.0f)
+		]
+		+ SHorizontalBox::Slot()
+		.AutoWidth()
+		[
+			PickWidget.ToSharedRef()
+		]
+	;
+
+	if (bClearButton)
+	{
+		PathPickerBox->AddSlot()
+			.AutoWidth()
+			[
+				ClearWidget.ToSharedRef()
+			];
+	}
+
+	FDetailWidgetRow& Row = Category.AddCustomRow(Label, false)
 		.NameContent()
 		[
 			SNew(SHorizontalBox)
@@ -338,22 +419,13 @@ void FLuminTargetSettingsDetails::BuildPathPicker(
 		]
 		.ValueContent()
 		[
-			SNew(SHorizontalBox)
-			+ SHorizontalBox::Slot()
-			.AutoWidth()
-			[
-				SNew(STextBlock)
-				.Text(TAttribute<FText>::Create(
-					TAttribute<FText>::FGetter::CreateLambda([Path]()->FText { return FText::FromString(Path.Get()); })))
-				.Font(DetailBuilder.GetDetailFont())
-				.Margin(2.0f)
-			]
-			+ SHorizontalBox::Slot()
-			.AutoWidth()
-			[
-				PickWidget.ToSharedRef()
-			]
+			PathPickerBox.ToSharedRef()
 		];
+
+	if (bDisableUntilConfigured)
+	{
+		Row.IsEnabled(SetupForPlatformAttribute);
+	}
 }
 
 FReply FLuminTargetSettingsDetails::OnPickDirectory(TAttribute<FString> DirPath, const FLuminTargetSettingsDetails::FOnPickPath& OnPick, TSharedPtr<SButton> PickButton)
diff --git a/Engine/Source/Developer/Lumin/LuminPlatformEditor/Private/LuminTargetSettingsDetails.h b/Engine/Source/Developer/Lumin/LuminPlatformEditor/Private/LuminTargetSettingsDetails.h
--- a/Engine/Source/Developer/Lumin/LuminPlatformEditor/Private/LuminTargetSettingsDetails.h
+++ b/Engine/Source/Developer/Lumin/LuminPlatformEditor/Private/LuminTargetSettingsDetails.h
@@ -39,7 +39,8 @@ private:
 	FUNC_DECLARE_DELEGATE(FOnChoosePath, FReply, TAttribute<FString>, const FOnPickPath&, TSharedPtr<SButton>);
 	void BuildPathPicker(IDetailLayoutBuilder& DetailBuilder, IDetailCategoryBuilder& Category,
 		TAttribute<FString> DirPath, const FText& Label, const FText& Tooltip,
-		const FOnChoosePath& OnChoose, const FOnPickPath& OnPick);
+		const FOnChoosePath& OnChoose, const FOnPickPath& OnPick, bool bClearButton, const FText& ClearTooltip,
+		bool bDisableUntilConfigured);
 	FReply OnPickDirectory(TAttribute<FString> DirPath,
 		const FOnPickPath& OnPick, TSharedPtr<SButton> PickButton);
 	FReply OnPickFile(TAttribute<FString> FilePath,
diff --git a/Engine/Source/Developer/Lumin/LuminPlatformEditor/Private/MagicLeapSDKSettings.h b/Engine/Source/Developer/Lumin/LuminPlatformEditor/Private/MagicLeapSDKSettings.h
--- a/Engine/Source/Developer/Lumin/LuminPlatformEditor/Private/MagicLeapSDKSettings.h
+++ b/Engine/Source/Developer/Lumin/LuminPlatformEditor/Private/MagicLeapSDKSettings.h
@@ -11,7 +11,7 @@ class IAndroidDeviceDetection;
 /**
  * Implements the settings for the Magic Leap SDK setup.
  */
-UCLASS(config = Engine, defaultconfig)
+UCLASS(config = Engine, globaluserconfig)
 class LUMINPLATFORMEDITOR_API UMagicLeapSDKSettings : public UObject
 {
 public:
